Fix static analysis warning
#rnx

diff --git a/Engine/Plugins/Runtime/NetworkPredictionExtras/Source/NetworkPredictionExtras/Private/NetworkPredictionExtrasFlyingPawn.cpp b/Engine/Plugins/Runtime/NetworkPredictionExtras/Source/NetworkPredictionExtras/Private/NetworkPredictionExtrasFlyingPawn.cpp
--- a/Engine/Plugins/Runtime/NetworkPredictionExtras/Source/NetworkPredictionExtras/Private/NetworkPredictionExtrasFlyingPawn.cpp
+++ b/Engine/Plugins/Runtime/NetworkPredictionExtras/Source/NetworkPredictionExtras/Private/NetworkPredictionExtrasFlyingPawn.cpp
@@ -48,13 +48,15 @@ void ANetworkPredictionExtrasFlyingPawn::BeginPlay()
 
 	if (UWorld* World = GetWorld())
 	{
-		ensure(FlyingMovementComponent);
-		FlyingMovementComponent->ProduceInputDelegate.BindUObject(this, &ANetworkPredictionExtrasFlyingPawn::ProduceInput);
+		if (ensure(FlyingMovementComponent))
+		{
+			FlyingMovementComponent->ProduceInputDelegate.BindUObject(this, &ANetworkPredictionExtrasFlyingPawn::ProduceInput);
+		}
 
 		// Binds 0 and 9 to the debug hud commands. This is just a convenience for the extras plugin. Real projects should bind this themselves
 		if (FlyingPawnCVars::BindAutomatically > 0)
 		{
-			if (ULocalPlayer* LocalPlayer = GetWorld()->GetFirstLocalPlayerFromController())
+			if (ULocalPlayer* LocalPlayer = World->GetFirstLocalPlayerFromController())
 			{
 				LocalPlayer->Exec(World, TEXT("setbind Nine nms.Debug.LocallyControlledPawn"), *GLog);
 				LocalPlayer->Exec(World, TEXT("setbind Zero nms.Debug.ToggleContinous"), *GLog);
