NetworkPrediction plugin:
-iteration on mock ability simulation implementation
-initial stub in of basic client side event handler

#rnx

diff --git a/Engine/Plugins/Runtime/NetworkPrediction/Source/NetworkPrediction/Public/NetworkedSimulationModelBuffer.h b/Engine/Plugins/Runtime/NetworkPrediction/Source/NetworkPrediction/Public/NetworkedSimulationModelBuffer.h
--- a/Engine/Plugins/Runtime/NetworkPrediction/Source/NetworkPrediction/Public/NetworkedSimulationModelBuffer.h
+++ b/Engine/Plugins/Runtime/NetworkPrediction/Source/NetworkPrediction/Public/NetworkedSimulationModelBuffer.h
@@ -37,7 +37,7 @@ struct TNetSimLazyWriterFunc
 	TFunction<void*()> GetFunc;
 };
 
-// Reference version of LazyWriter. Faster to pass around and do cascading downcasts
+// Reference version of LazyWriter. This is what gets passed through chains of ::SimulationTick calls. This is to avoid copying the TFunction in TNetSimLazyWriterFunc around.
 template<typename ElementType>
 struct TNetSimLazyWriter
 {
@@ -87,7 +87,7 @@ protected:
 	template<typename ElementType>
 	ElementType* WriteFrameInitializedFromHeadImpl(int32 Frame)
 	{
-		if (((T*)this)->HeadFrame() > Frame)
+		if (Frame > ((T*)this)->HeadFrame())
 		{
 			ElementType* HeadElementPtr = ((T*)this)->HeadElement();
 			ElementType* NewElement = ((T*)this)->WriteFrame(Frame);
diff --git a/Engine/Plugins/Runtime/NetworkPrediction/Source/NetworkPrediction/Public/NetworkedSimulationModelTypes.h b/Engine/Plugins/Runtime/NetworkPrediction/Source/NetworkPrediction/Public/NetworkedSimulationModelTypes.h
--- a/Engine/Plugins/Runtime/NetworkPrediction/Source/NetworkPrediction/Public/NetworkedSimulationModelTypes.h
+++ b/Engine/Plugins/Runtime/NetworkPrediction/Source/NetworkPrediction/Public/NetworkedSimulationModelTypes.h
@@ -331,6 +331,8 @@ struct FNetworkSimTime
 	// FIXME
 	void NetSerialize(FArchive& Ar) { Ar << Time; }
 
+	operator FSimTime() const { return Time; }
+
 	using T = FNetworkSimTime;
 	T& operator+= (const T &rhs) { this->Time += rhs.Time; return(*this); }
 	T& operator-= (const T &rhs) { this->Time -= rhs.Time; return(*this); }
diff --git a/Engine/Plugins/Runtime/NetworkPredictionExtras/Source/NetworkPredictionExtras/Private/MockAbilitySimulation.cpp b/Engine/Plugins/Runtime/NetworkPredictionExtras/Source/NetworkPredictionExtras/Private/MockAbilitySimulation.cpp
--- a/Engine/Plugins/Runtime/NetworkPredictionExtras/Source/NetworkPredictionExtras/Private/MockAbilitySimulation.cpp
+++ b/Engine/Plugins/Runtime/NetworkPredictionExtras/Source/NetworkPredictionExtras/Private/MockAbilitySimulation.cpp
@@ -23,124 +23,200 @@ const FName FMockAbilitySimulation::GroupName(TEXT("Ability"));
 
 void FMockAbilitySimulation::SimulationTick(const TNetSimTimeStep& TimeStep, const TNetSimInput<TMockAbilityBufferTypes>& Input, const TNetSimOutput<TMockAbilityBufferTypes>& Output)
 {
+	// WIP NOTES:
+	//	-We are creating local copies of the input state that we can mutate and pass on to the parent simulation's tick (LocalCmd, LocalSync, LocalAux). This may seem weird but consider:
+	//		-Modifying the actual inputs is a bad idea: inputs are "final" once TNetworkedSimulationModel calls SimulationTick. The sim itself should not modify the inputs after this.
+	//		-Tempting to just use the already-allocated Output states and pass them to the parent sim as both Input/Output. But now &Input.Sync == &Output.Sync for example! The parent sim
+	//			may write to Output and then later check it against the passed in input (think "did this change"), having no idea that by writing to output he was also writing to input! This seems like it should be avoided.
+	//		-It may be worth considering "lazy copying" or allocating of these temp states... but for now that adds more complexity and isn't necessary since the local copies are stack allocated and
+	//			these state structures are small and contiguous.
+	//
+	//	-Still, there are some awkwardness when we need to write to both the local temp state *and* the output state explicitly.
+	//	-Stamina changes for example: we want to both change the stamina both for this frame (used by rest of code in this function) AND write it to the output state (so it comes in as input next frame)
+	//		-Since this sim "owns" stamina, we could copy it immediately to the output state and use the output state exclusively in this code...
+	//		-... but that is inconsistent with state that we don't "own" like Location. Since the parent sim will generate a new Location and we can't pass output to the parent as input
+	//			-You would be in a situation where "Teleport modified LocalSync.Location" and "another place modified Output.Sync.Stamina".
+	//	
+	//
+	//	Consider: not using inheritance for simulations. Would that make this less akward? 
+	//		-May make Sync state case above less weird since it would breakup stamina/location ("owned" vs "not owned"). Stamina could always use output state and we'd have a local copy of the movement input sync state.
+	//		-Aux state is still bad even without inheritance: if you want to read and write to an aux variable at different points in the frames, hard to know where "real" value is at any given time without
+	//			writing fragile code. Making code be explicit: "I am writing to the current value this frame AND next frames value" seems ok. Could possibly do everything in a local copy and then see if its changed
+	//			at the end up ::SimulationTick and then copy it into output aux. Probably should avoid comparison operator each frame so would need to wrap in some struct that could track a dirty flag... seems complicating.
+	//
+
+	check(EventHandler);
 	const float DeltaSeconds = TimeStep.StepMS.ToRealTimeSeconds();
 
-	// Stamina passes through. Some code paths will modify this again, but if we don't set the output state it will be garbage/stale
-	// (considering implicit copying of old state to new state by the NetworkedSimModel code, but that could be undesired/inefficient in some cases)
-	Output.Sync.Stamina = Input.Sync.Stamina;
+	// Local copies of the const input that we will pass into the parent sim as input.
+	FMockAbilityInputCmd LocalCmd = Input.Cmd;
+	FMockAbilitySyncState LocalSync = Input.Sync;
+	FMockAbilityAuxstate LocalAux = Input.Aux;
 
-	static float BlinkCost = 75.f;
-	const bool bBlink = (Input.Cmd.bBlinkPressed && Input.Sync.Stamina > BlinkCost);
+	const bool bAlreadyBlinking = Input.Aux.BlinkWarmupLeft > 0;
+	const bool bAlreadyDashing = Input.Aux.DashTimeLeft > 0;
+	const bool bAllowNewActivations = !bAlreadyBlinking && !bAlreadyDashing;
+
+	// -------------------------------------------------------------------------
+	//	Regen
+	//	-Applies at the start of frame: you can spend it immediately
+	//	-hence we have to write to both the LocalSync's stamina (what the rest of the abilities look at for input) and output (what value will be used as input for next frame)
+	// -------------------------------------------------------------------------
 
-	if (bBlink)
 	{
-		static float BlinkDist = 1000.f;
-		static float OverlapCheckInflation = 0.100f;
+		const float NewStamina = FMath::Min<float>( LocalSync.Stamina + (DeltaSeconds * LocalAux.StaminaRegenRate), LocalAux.MaxStamina );
+		LocalSync.Stamina = NewStamina;
+		Output.Sync.Stamina = NewStamina;
+	}
 
-		FVector DestLocation = Input.Sync.Location + Input.Sync.Rotation.RotateVector( FVector(BlinkDist, 0.f, 0.f) );
-		AActor* OwningActor = UpdatedComponent->GetOwner();
-		check(OwningActor);
+	// -------------------------------------------------------------------------
+	//	Blink
+	// -------------------------------------------------------------------------
 
-		// DrawDebugLine(OwningActor->GetWorld(), Input.Sync.Location, DestLocation, FColor::Red, false);	
+	static float BlinkCost = 75.f;
+	static int16 BlinkWarmupMS = 750; 
 
-		// Its unfortunate teleporting is so complicated. It may make sense for a new movement simulation to define this themselves, 
-		// but for this mock one, we will just use the engine's AActor teleport.
-		if (OwningActor->TeleportTo(DestLocation, Input.Sync.Rotation))
-		{
-			Output.Sync = Input.Sync;
+	const bool bBlinkActivate = (Input.Cmd.bBlinkPressed && Input.Sync.Stamina > BlinkCost && bAllowNewActivations);
+	if (bBlinkActivate)
+	{
+		LocalAux.BlinkWarmupLeft = BlinkWarmupMS;
+		EventHandler->NotifyBlinkStartup();
+	}
+
+	if (LocalAux.BlinkWarmupLeft > 0)
+	{
+		const int16 NewBlinkWarmupLeft = FMath::Max<int16>(0, LocalAux.BlinkWarmupLeft - TimeStep.StepMS);
+		Output.Aux.Get()->BlinkWarmupLeft = NewBlinkWarmupLeft;
+
+		// While blinking is warming up, we disable movement input and the other actions
+		LocalCmd.MovementInput.Set(0.f, 0.f, 0.f);
+		LocalCmd.bDashPressed = false;
+		LocalCmd.bSprintPressed = false;
 
-			// Component now has the final location
-			const FTransform UpdateComponentTransform = GetUpdateComponentTransform();
-			Output.Sync.Location = UpdateComponentTransform.GetLocation();
-			Output.Sync.Stamina = Input.Sync.Stamina - BlinkCost;
+		LocalSync.Velocity.Set(0.f, 0.f, 0.f);
 
-			// And we skip the normal update simulation for this frame. This is just a choice. We could still allow it to run.
-			return;
+		if (NewBlinkWarmupLeft <= 0)
+		{
+			static float BlinkDist = 1000.f;
+			FVector DestLocation = LocalSync.Location + LocalSync.Rotation.RotateVector( FVector(BlinkDist, 0.f, 0.f) );
+			AActor* OwningActor = UpdatedComponent->GetOwner();
+			check(OwningActor);
+
+			// DrawDebugLine(OwningActor->GetWorld(), Input.Sync.Location, DestLocation, FColor::Red, false);	
+
+			// Its unfortunate teleporting is so complicated. It may make sense for a new movement simulation to define this themselves, 
+			// but for this mock one, we will just use the engine's AActor teleport.
+			if (OwningActor->TeleportTo(DestLocation, LocalSync.Rotation))
+			{
+				// Component now has the final location
+				const FTransform UpdateComponentTransform = GetUpdateComponentTransform();
+				LocalSync.Location = UpdateComponentTransform.GetLocation();
+				
+				const float NewStamina = FMath::Max<float>(0.f, LocalSync.Stamina - BlinkCost);
+				LocalSync.Stamina = NewStamina;
+				Output.Sync.Stamina = NewStamina;
+
+				EventHandler->NotifyBlinkFinished();
+			}
 		}
 	}
 
-	// Dash is implemented in the following way:
+	// -------------------------------------------------------------------------
+	//	Dash
 	//	-Stamina consumed on initial press
 	//	-MaxSpeed/Acceleration are jacked up
 	//	-Dash lasts for 400ms (DashDurationMS)
-	//		-division of frame times can cause you to dash for longer. We would have to break up simulation steps to support this 100% accurately.
-	//	-Movement input is synthesized while in dash state. That is, we force forward movement and ignore what was actually fed into the simulation.
-	//
-	// This is just a simple/interesting way of implementing dash in the this system. A real movement/ability system will probably have some concept of
-	// root motion/sources that drive movement forward.
+	//		-Division of frame times can cause you to dash for longer. We would have to break up simulation steps to support this 100% accurately.
+	//	-Movement input is synthesized while in dash state. That is, we force forward movement and ignore what was actually fed into the simulation (move input only)
+	// -------------------------------------------------------------------------
 
 	static float DashCost = 75.f;
-	static int16 DashDurationMS = 400;
-	int16 DashTimeLeft = Input.Aux.DashTimeLeft;
-	bool bIsDashing = Input.Aux.DashTimeLeft > 0;
+	static int16 DashDurationMS = 500;
 	
-	if (Input.Cmd.bDashPressed && Input.Sync.Stamina > DashCost && bIsDashing == false)
+	const bool bDashActivate = (Input.Cmd.bDashPressed && Input.Sync.Stamina > DashCost && bAllowNewActivations);
+	if (bDashActivate)
 	{
 		// Start dashing
-		DashTimeLeft = DashDurationMS;
-		Output.Sync.Stamina -= DashCost;
-		bIsDashing = true;
+		LocalAux.DashTimeLeft = DashDurationMS;
+		
+		const float NewStamina = FMath::Max<float>(0.f, LocalSync.Stamina - DashCost);
+		LocalSync.Stamina = NewStamina;
+		Output.Sync.Stamina = NewStamina;
+
+		EventHandler->NotifyDash(true);
 	}
 
-	if (bIsDashing)
+	if (LocalAux.DashTimeLeft > 0)
 	{
-		FMockAbilityAuxstate* OutAuxState = Output.Aux.Get();
-		OutAuxState->DashTimeLeft = FMath::Max<int16>(DashTimeLeft - (int16)(DeltaSeconds * 1000.f), 0);
+		const int16 NewDashTimeLeft = FMath::Max<int16>(LocalAux.DashTimeLeft - (int16)(DeltaSeconds * 1000.f), 0);
+		Output.Aux.Get()->DashTimeLeft = NewDashTimeLeft;
 
-		FMockAbilityAuxstate LocalAuxState = Input.Aux;
-		LocalAuxState.MaxSpeed = MockAbilityCVars::DashMaxSpeed();
-		LocalAuxState.Acceleration = MockAbilityCVars::DashAcceleration();
+		LocalAux.MaxSpeed = MockAbilityCVars::DashMaxSpeed();
+		LocalAux.Acceleration = MockAbilityCVars::DashAcceleration();
 
-		FMockAbilityInputCmd LocalInputCmd = Input.Cmd;
-		LocalInputCmd.MovementInput = FVector(1.f, 0.f, 0.f);
+		LocalCmd.MovementInput.Set(1.f, 0.f, 0.f);
+		LocalCmd.bBlinkPressed = false;
+		LocalCmd.bSprintPressed = false;
 		
-		FlyingMovement::FMovementSimulation::SimulationTick(TimeStep, { LocalInputCmd, Input.Sync, LocalAuxState }, { Output.Sync, Output.Aux });
+		FlyingMovement::FMovementSimulation::SimulationTick(TimeStep, { LocalCmd, LocalSync, LocalAux }, { Output.Sync, Output.Aux });
 
-		if (OutAuxState->DashTimeLeft == 0)
+		if (NewDashTimeLeft == 0)
 		{
 			// Stop when dash is over
-			Output.Sync.Velocity = FVector::ZeroVector;
+			Output.Sync.Velocity.Set(0.f, 0.f, 0.f);
+			EventHandler->NotifyDash(false);
 		}
+
+		return;
 	}
-	else
+
+
+	// -------------------------------------------------------------------------
+	//	Sprint
+	//	-Requires 15% stamina to start activation
+	//	-Drains at 100 stamina/second after that
+	//	-Just bumps up max move speed
+	//	-Note how this is a transient value: it is fed into the parent sim's input but never makes it to an output state
+	// -------------------------------------------------------------------------
+	
+	static float SprintBaseCost = 100.f; // 100 sprint/second cost
+	static float SprintStartMinPCT = 0.15f; // 15% stamina required to begin sprinting
+	
+	bool bIsSprinting = false;
+	if (Input.Cmd.bSprintPressed && bAllowNewActivations)
 	{
-		// Sprint (mutually exclusive from Dash state)		
-		static float SprintBaseCost = 100.f;
 		const float SprintCostThisFrame = SprintBaseCost * DeltaSeconds;
-		const bool bIsSprinting = (Input.Cmd.bSprintPressed && Input.Sync.Stamina > SprintCostThisFrame);
-
-		// Set our max speed. This is an interesting case.
-		//	-Our input states are already "final". It doesn't make sense to modify the input AuxState data.
-		//	-But, we want to feed this locally calculated max speed into the base movement simulation. So, it is an "input" in that sense.
-		//	-Creating a local copy of the aux state, modifying it, and passing that into the base movement sim is clean, though a bit inefficient.
-		//	-This leaves us with a weird side effect that the calculated MaxSpeed is never written to the AuxBuffer. Though we could write it to the
-		//	output AuxState, that doesn't really do anything useful: the base move sim won't use it and it'll just be left there for next frame (to be overriden again).
-		//
-		// The way the base movement simulation is written, MaxSpeed is a clearly defined input. But in the MockAbility sim, MaxSpeed is a derived
-		// value from other input state. This difference in how the simulations treat the variable is what causes the need to do this.
-		//
-		// It would be possible to write the base movement sim in a way that MaxSpeed is transient value on the sim class (FlyingMovement::FMovementSimulation).
-		// Something like, "MaxSpeed" really means "Base max speed" and there would be a "GetCurrentMaxSpeed" virtual on the sim.
-		// This would make things a bit more awkward in the base case with no ability system. So, for now, this seems like a good pattern/precedence.
-
-		FMockAbilityAuxstate LocalAuxState = Input.Aux;
-		LocalAuxState.MaxSpeed = bIsSprinting ? MockAbilityCVars::SprintMaxSpeed() : MockAbilityCVars::DefaultMaxSpeed();
 
-		if (bIsSprinting)
+		if (Input.Aux.bIsSprinting)
 		{
-			Output.Sync.Stamina = FMath::Max<float>( Input.Sync.Stamina - SprintCostThisFrame, 0.f );
+			// If we are already sprinting, then we keep sprinting as long as we can afford the absolute cost
+			bIsSprinting = LocalSync.Stamina > SprintCostThisFrame;
 		}
-		else if (Output.Sync.Stamina < Input.Aux.MaxStamina)
+		else
 		{
-			Output.Sync.Stamina = FMath::Min<float>( Input.Sync.Stamina + (DeltaSeconds * Input.Aux.StaminaRegenRate), Input.Aux.MaxStamina );
+			// To start sprinting, we have to have SprintStartMinPCT stamina
+			bIsSprinting = LocalSync.Stamina > (SprintStartMinPCT * LocalAux.MaxStamina);
 		}
-		else
+
+		if (bIsSprinting)
 		{
-			Output.Sync.Stamina = Input.Sync.Stamina;
+			LocalAux.MaxSpeed = MockAbilityCVars::SprintMaxSpeed();
+
+			const float NewStamina = FMath::Max<float>(0.f, LocalSync.Stamina - SprintCostThisFrame);
+			LocalSync.Stamina = NewStamina;
+			Output.Sync.Stamina = NewStamina;
+
 		}
+	}
 
-		FlyingMovement::FMovementSimulation::SimulationTick(TimeStep, Input, { Output.Sync, Output.Aux });
+	// Update the out aux state and call notifies only if sprinting state as actually changed
+	if (bIsSprinting != Input.Aux.bIsSprinting)
+	{
+		EventHandler->NotifySprint(bIsSprinting);
+		Output.Aux.Get()->bIsSprinting = bIsSprinting;
 	}
+
+	FlyingMovement::FMovementSimulation::SimulationTick(TimeStep, { LocalCmd, LocalSync, LocalAux }, { Output.Sync, Output.Aux });
 }
 
 // -------------------------------------------------------------------------------------------------------------
@@ -160,6 +236,7 @@ INetworkSimulationModel* UMockFlyingAbilityComponent::InstantiateNetworkSimulati
 	MockAbilitySimulation.Reset(new FMockAbilitySimulation());
 	MockAbilitySimulation->UpdatedComponent = UpdatedComponent;
 	MockAbilitySimulation->UpdatedPrimitive = UpdatedPrimitive;
+	MockAbilitySimulation->EventHandler = this;
 	
 	FMockAbilitySyncState InitialSyncState;
 	InitialSyncState.Location = UpdatedComponent->GetComponentLocation();
@@ -202,3 +279,43 @@ void UMockFlyingAbilityComponent::VisualLog(const FMockAbilityInputCmd* Input, c
 }
 
 // ---------------------------------------------------------------------------------
+
+void UMockFlyingAbilityComponent::NotifySprint(bool bIsSprinting)
+{
+	if (GetNetMode() == NM_DedicatedServer)
+	{
+		return;
+	}
+
+	UE_LOG(LogTemp, Display, TEXT("Sprint: %d"), bIsSprinting);
+}
+
+void UMockFlyingAbilityComponent::NotifyDash(bool bIsDashing)
+{
+	if (GetNetMode() == NM_DedicatedServer)
+	{
+		return;
+	}
+
+	UE_LOG(LogTemp, Display, TEXT("Dash: %d"), bIsDashing);
+}
+
+void UMockFlyingAbilityComponent::NotifyBlinkStartup()
+{
+	if (GetNetMode() == NM_DedicatedServer)
+	{
+		return;
+	}
+
+	UE_LOG(LogTemp, Display, TEXT("Blink Startup"));
+}
+
+void UMockFlyingAbilityComponent::NotifyBlinkFinished()
+{
+	if (GetNetMode() == NM_DedicatedServer)
+	{
+		return;
+	}
+
+	UE_LOG(LogTemp, Display, TEXT("Blink Finished"));
+}
\ No newline at end of file
diff --git a/Engine/Plugins/Runtime/NetworkPredictionExtras/Source/NetworkPredictionExtras/Public/MockAbilitySimulation.h b/Engine/Plugins/Runtime/NetworkPredictionExtras/Source/NetworkPredictionExtras/Public/MockAbilitySimulation.h
--- a/Engine/Plugins/Runtime/NetworkPredictionExtras/Source/NetworkPredictionExtras/Public/MockAbilitySimulation.h
+++ b/Engine/Plugins/Runtime/NetworkPredictionExtras/Source/NetworkPredictionExtras/Public/MockAbilitySimulation.h
@@ -90,12 +90,16 @@ struct FMockAbilityAuxstate : public FlyingMovement::FAuxState
 	float MaxStamina = 100.f;
 	float StaminaRegenRate = 20.f;
 	int16 DashTimeLeft = 0;
+	int16 BlinkWarmupLeft = 0;
+	bool bIsSprinting = false;
 
 	void NetSerialize(const FNetSerializeParams& P)
 	{
 		P.Ar << MaxStamina;
 		P.Ar << StaminaRegenRate;
 		P.Ar << DashTimeLeft;
+		P.Ar << BlinkWarmupLeft;
+		P.Ar << bIsSprinting;
 		FlyingMovement::FAuxState::NetSerialize(P);
 	}
 
@@ -106,13 +110,27 @@ struct FMockAbilityAuxstate : public FlyingMovement::FAuxState
 		{
 			P.Ar->Logf(TEXT("MaxStamina: %.2f"), MaxStamina);
 			P.Ar->Logf(TEXT("StaminaRegenRate: %.2f"), StaminaRegenRate);
-			P.Ar->Logf(TEXT("StaminaRegenRate: %.2f"), DashTimeLeft);
+			P.Ar->Logf(TEXT("DashTimeLeft: %d"), DashTimeLeft);
+			P.Ar->Logf(TEXT("BlinkWarmupLeft: %d"), BlinkWarmupLeft);
+			P.Ar->Logf(TEXT("bIsSprinting: %d"), bIsSprinting);
 		}
 	}
 };
 
 using TMockAbilityBufferTypes = TNetworkSimBufferTypes<FMockAbilityInputCmd, FMockAbilitySyncState, FMockAbilityAuxstate>;
 
+// A very hard coded event handling interface for events that are emitted in FMockAbilitySimulation::SimulationTick.
+// These events are for cosmetic, client side work. E.g, we are not implementing sprint/dash/blink here, we are doing
+// visual/audio effects based on those states changing within the core simulation update. These are effectively skipped on the server (in implementation).
+class IMockAbilityEventHandler
+{
+public:
+	virtual void NotifySprint(bool bIsSprinting) = 0;
+	virtual void NotifyDash(bool bIsDashing) = 0;
+	virtual void NotifyBlinkStartup() = 0;
+	virtual void NotifyBlinkFinished() = 0;
+};
+
 class FMockAbilitySimulation : public FlyingMovement::FMovementSimulation
 {
 public:
@@ -121,6 +139,8 @@ public:
 
 	/** Main update function */
 	void SimulationTick(const TNetSimTimeStep& TimeStep, const TNetSimInput<TMockAbilityBufferTypes>& Input, const TNetSimOutput<TMockAbilityBufferTypes>& Output);
+
+	IMockAbilityEventHandler* EventHandler = nullptr;
 };
 
 template<int32 InFixedStepMS=0>
@@ -133,7 +153,7 @@ class IMockFlyingAbilitySystemDriver : public TNetworkedSimulationModelDriver<TM
 // -------------------------------------------------------------------------------------------------------------------------------
 
 UCLASS(BlueprintType, meta=(BlueprintSpawnableComponent))
-class NETWORKPREDICTIONEXTRAS_API UMockFlyingAbilityComponent : public UFlyingMovementComponent, public IMockFlyingAbilitySystemDriver
+class NETWORKPREDICTIONEXTRAS_API UMockFlyingAbilityComponent : public UFlyingMovementComponent, public IMockFlyingAbilitySystemDriver, public IMockAbilityEventHandler
 {
 	GENERATED_BODY()
 
@@ -160,6 +180,12 @@ public:
 	using UFlyingMovementComponent::ProduceInput;
 	using UFlyingMovementComponent::FinalizeFrame;
 
+	// IMockAbilityEventHandler
+	virtual void NotifySprint(bool bIsSprinting) override;
+	virtual void NotifyDash(bool bIsDashing) override;
+	virtual void NotifyBlinkStartup() override;
+	virtual void NotifyBlinkFinished() override;
+
 protected:
 
 	// Network Prediction
