NetworkPrediction plugin
-Fix bp compile error in level script
-adjust mock ability costs

#rnx

diff --git a/Engine/Plugins/Runtime/NetworkPredictionExtras/Source/NetworkPredictionExtras/Private/MockAbilitySimulation.cpp b/Engine/Plugins/Runtime/NetworkPredictionExtras/Source/NetworkPredictionExtras/Private/MockAbilitySimulation.cpp
--- a/Engine/Plugins/Runtime/NetworkPredictionExtras/Source/NetworkPredictionExtras/Private/MockAbilitySimulation.cpp
+++ b/Engine/Plugins/Runtime/NetworkPredictionExtras/Source/NetworkPredictionExtras/Private/MockAbilitySimulation.cpp
@@ -73,7 +73,7 @@ void FMockAbilitySimulation::SimulationTick(const TNetSimTimeStep& TimeStep, con
 	//	Blink
 	// -------------------------------------------------------------------------
 
-	static float BlinkCost = 75.f;
+	static float BlinkCost = 25.f;
 	static int16 BlinkWarmupMS = 750; 
 
 	const bool bBlinkActivate = (Input.Cmd.bBlinkPressed && Input.Sync.Stamina > BlinkCost && bAllowNewActivations);
@@ -130,7 +130,7 @@ void FMockAbilitySimulation::SimulationTick(const TNetSimTimeStep& TimeStep, con
 	//	-Movement input is synthesized while in dash state. That is, we force forward movement and ignore what was actually fed into the simulation (move input only)
 	// -------------------------------------------------------------------------
 
-	static float DashCost = 75.f;
+	static float DashCost = 25.f;
 	static int16 DashDurationMS = 500;
 	
 	const bool bDashActivate = (Input.Cmd.bDashPressed && Input.Sync.Stamina > DashCost && bAllowNewActivations);
