Network Prediction Plugin
-Added Cue example to Mock Simulation

#rnx

diff --git a/Engine/Plugins/Runtime/NetworkPredictionExtras/Source/NetworkPredictionExtras/Private/MockNetworkSimulation.cpp b/Engine/Plugins/Runtime/NetworkPredictionExtras/Source/NetworkPredictionExtras/Private/MockNetworkSimulation.cpp
--- a/Engine/Plugins/Runtime/NetworkPredictionExtras/Source/NetworkPredictionExtras/Private/MockNetworkSimulation.cpp
+++ b/Engine/Plugins/Runtime/NetworkPredictionExtras/Source/NetworkPredictionExtras/Private/MockNetworkSimulation.cpp
@@ -58,9 +58,11 @@ static FAutoConsoleVariableRef CVarBindAutomatically(TEXT("mns.BindAutomatically
 // ============================================================================================
 
 static bool ForceMispredict = false;
-
 const FName FMockNetworkModelDef::GroupName(TEXT("Mock"));
 
+NETSIMCUE_REGISTER(FMockCue, TEXT("MockCue"));
+NETSIMCUESET_REGISTER(UMockNetworkSimulationComponent, FMockCueSet);
+
 void FMockNetworkSimulation::SimulationTick(const FNetSimTimeStep& TimeStep, const TNetSimInput<TMockNetworkSimulationBufferTypes>& Input, const TNetSimOutput<TMockNetworkSimulationBufferTypes>& Output)
 {
 	Output.Sync.Total = Input.Sync.Total + (Input.Cmd.InputValue * Input.Aux.Multiplier * TimeStep.StepMS.ToRealTimeSeconds());
@@ -74,6 +76,13 @@ void FMockNetworkSimulation::SimulationTick(const FNetSimTimeStep& TimeStep, con
 		
 		ForceMispredict = false;
 	}
+	
+	if ((int32)(Input.Sync.Total/10.f) < (int32)(Output.Sync.Total/10.f))
+	{
+		// Emit a cue for crossing this threshold. Note this could mostly be accomplished with a state transition (by detecting this in FinalizeFrame)
+		// But to illustrate Cues, we are adding a random number to the cue's payload. The point being cues can transmit data that is not part of the sync/aux state.
+		Output.CueDispatch.Invoke<FMockCue>( FMath::Rand() % 1024 );
+	}
 }
 
 // ===============================================================================================
@@ -120,6 +129,8 @@ INetworkedSimulationModel* UMockNetworkSimulationComponent::InstantiateNetworked
 
 void UMockNetworkSimulationComponent::TickComponent(float DeltaTime, enum ELevelTick TickType, FActorComponentTickFunction *ThisTickFunction)
 {
+	Super::TickComponent(DeltaTime, TickType, ThisTickFunction);
+
 	const ENetRole OwnerRole = GetOwnerRole();
 
 	// Check if we should trip a mispredict. (Note how its not possible to do this inside the Update function!)
@@ -171,6 +182,11 @@ void UMockNetworkSimulationComponent::VisualLog(const FMockInputCmd* Input, cons
 	FVisualLoggingHelpers::VisualLogActor(Owner, Transform, LocalParameters);
 }
 
+void UMockNetworkSimulationComponent::HandleCue(const FMockCue& MockCue, const FNetSimCueSystemParamemters& SystemParameters)
+{
+	UE_LOG(LogMockNetworkSim, Display, TEXT("MockCue Handled! Value: %d"), MockCue.RandomData);
+}
+
 // ------------------------------------------------------------------------------------------------------------------------
 //
 // ------------------------------------------------------------------------------------------------------------------------
diff --git a/Engine/Plugins/Runtime/NetworkPredictionExtras/Source/NetworkPredictionExtras/Public/MockNetworkSimulation.h b/Engine/Plugins/Runtime/NetworkPredictionExtras/Source/NetworkPredictionExtras/Public/MockNetworkSimulation.h
--- a/Engine/Plugins/Runtime/NetworkPredictionExtras/Source/NetworkPredictionExtras/Public/MockNetworkSimulation.h
+++ b/Engine/Plugins/Runtime/NetworkPredictionExtras/Source/NetworkPredictionExtras/Public/MockNetworkSimulation.h
@@ -49,6 +49,11 @@ class IMockDriver;
 //
 // -------------------------------------------------------------------------------------------------------------------------------
 
+
+// -------------------------------------------------------------------------------------------------------------------------------
+//	Simulation State
+// -------------------------------------------------------------------------------------------------------------------------------
+
 // State the client generates
 struct FMockInputCmd
 {
@@ -107,7 +112,6 @@ struct FMockSyncState
 
 // Auxiliary state that is input into the simulation. Doesn't change during the simulation tick.
 // (It can change and even be predicted but doing so will trigger more bookeeping, etc. Changes will happen "next tick")
-// NOTE: Currently incomplete!
 struct FMockAuxState
 {
 	float Multiplier=1;
@@ -129,6 +133,35 @@ struct FMockAuxState
 	}
 };
 
+// -------------------------------------------------------------------------------------------------------------------------------
+//	Simulation Cues
+// -------------------------------------------------------------------------------------------------------------------------------
+
+// A minimal Cue, emitted every 10 "totals" with a random integer as a payload
+struct FMockCue
+{
+	FMockCue() = default;
+	FMockCue(int32 InRand)
+		: RandomData(InRand) { }
+
+	NETSIMCUE_BODY();
+	int32 RandomData = 0;
+};
+
+// Set of all cues the Mock simulation emits. Not strictly necessary and not that helpful when there is only one cue (but using since this is the example real simulations will want to use)
+struct FMockCueSet
+{
+	template<typename TDispatchTable>
+	static void RegisterNetSimCueTypes(TDispatchTable& DispatchTable)
+	{
+		DispatchTable.template RegisterType<FMockCue>();
+	}
+};
+
+// -------------------------------------------------------------------------------------------------------------------------------
+//	The Simulation
+// -------------------------------------------------------------------------------------------------------------------------------
+
 using TMockNetworkSimulationBufferTypes = TNetworkSimBufferTypes<FMockInputCmd, FMockSyncState, FMockAuxState>;
 
 class FMockNetworkSimulation
@@ -139,6 +172,10 @@ public:
 	void SimulationTick(const FNetSimTimeStep& TimeStep, const TNetSimInput<TMockNetworkSimulationBufferTypes>& Input, const TNetSimOutput<TMockNetworkSimulationBufferTypes>& Output);
 };
 
+// -------------------------------------------------------------------------------------------------------------------------------
+//	Networked Simulation Model Def
+// -------------------------------------------------------------------------------------------------------------------------------
+
 class FMockNetworkModelDef : public FNetSimModelDefBase
 {
 public:
@@ -163,7 +200,7 @@ public:
 };
 
 // -------------------------------------------------------------------------------------------------------------------------------
-// ActorComponent for running a MockNetworkSimulation
+// ActorComponent for running a MockNetworkSimulation (implements Driver for the mock simulation)
 // -------------------------------------------------------------------------------------------------------------------------------
 
 UCLASS(BlueprintType, meta=(BlueprintSpawnableComponent))
@@ -186,7 +223,9 @@ public:
 	void VisualLog(const FMockInputCmd* Input, const FMockSyncState* Sync, const FMockAuxState* Aux, const FVisualLoggingParameters& SystemParameters) const override;
 	
 	void ProduceInput(const FNetworkSimTime SimTime, FMockInputCmd& Cmd);
-	void FinalizeFrame(const FMockSyncState& SyncState, const FMockAuxState& AuxState) override;	
+	void FinalizeFrame(const FMockSyncState& SyncState, const FMockAuxState& AuxState) override;
+
+	void HandleCue(const FMockCue& MockCue, const FNetSimCueSystemParamemters& SystemParameters);
 
 	// Mock representation of "syncing' to the sync state in the network sim.
 	float MockValue = 1000.f;
