Remove unnecessary modules that were causing warnings in non editor builds
#rnx

diff --git a/Engine/Plugins/Runtime/NetworkPrediction/Source/NetworkPrediction/NetworkPrediction.Build.cs b/Engine/Plugins/Runtime/NetworkPrediction/Source/NetworkPrediction/NetworkPrediction.Build.cs
--- a/Engine/Plugins/Runtime/NetworkPrediction/Source/NetworkPrediction/NetworkPrediction.Build.cs
+++ b/Engine/Plugins/Runtime/NetworkPrediction/Source/NetworkPrediction/NetworkPrediction.Build.cs
@@ -32,10 +32,7 @@ namespace UnrealBuildTool.Rules
 			PrivateDependencyModuleNames.AddRange(
 				new string[]
 				{
-					"TraceLog",
-					"TraceAnalysis",
-					"TraceServices",
-					"TraceInsights",
+					"TraceLog"					
                 }
 				);
 
diff --git a/Engine/Plugins/Runtime/NetworkPredictionExtras/Source/NetworkPredictionExtras/NetworkPredictionExtras.Build.cs b/Engine/Plugins/Runtime/NetworkPredictionExtras/Source/NetworkPredictionExtras/NetworkPredictionExtras.Build.cs
--- a/Engine/Plugins/Runtime/NetworkPredictionExtras/Source/NetworkPredictionExtras/NetworkPredictionExtras.Build.cs
+++ b/Engine/Plugins/Runtime/NetworkPredictionExtras/Source/NetworkPredictionExtras/NetworkPredictionExtras.Build.cs
@@ -33,10 +33,7 @@ namespace UnrealBuildTool.Rules
 			PrivateDependencyModuleNames.AddRange(
 				new string[]
 				{
-					"TraceLog",
-                    "TraceAnalysis",
-                    "TraceServices",
-                    "TraceInsights",
+					"TraceLog"                    
                 }
 				);
 
