Fixing broken build

diff --git a/Engine/Plugins/Runtime/GoogleVR/GoogleVRController/Source/GoogleVRController/Private/GoogleVRPointerInputComponent.cpp b/Engine/Plugins/Runtime/GoogleVR/GoogleVRController/Source/GoogleVRController/Private/GoogleVRPointerInputComponent.cpp
--- a/Engine/Plugins/Runtime/GoogleVR/GoogleVRController/Source/GoogleVRController/Private/GoogleVRPointerInputComponent.cpp
+++ b/Engine/Plugins/Runtime/GoogleVR/GoogleVRController/Source/GoogleVRController/Private/GoogleVRPointerInputComponent.cpp
@@ -45,7 +45,7 @@ bool UGoogleVRPointerInputComponent::IsBlockingHit() const
 
 AActor* UGoogleVRPointerInputComponent::GetHitActor() const
 {
-	return LatestHitResult.GetActor();
+	return LatestHitResult.GetHitObjectHandle().FetchActor();
 }
 
 UPrimitiveComponent* UGoogleVRPointerInputComponent::GetHitComponent() const
diff --git a/Engine/Plugins/Runtime/NetworkPredictionExtras/Source/NetworkPredictionExtras/Private/BaseMovementSimulation.cpp b/Engine/Plugins/Runtime/NetworkPredictionExtras/Source/NetworkPredictionExtras/Private/BaseMovementSimulation.cpp
--- a/Engine/Plugins/Runtime/NetworkPredictionExtras/Source/NetworkPredictionExtras/Private/BaseMovementSimulation.cpp
+++ b/Engine/Plugins/Runtime/NetworkPredictionExtras/Source/NetworkPredictionExtras/Private/BaseMovementSimulation.cpp
@@ -88,7 +88,7 @@ bool FBaseMovementSimulation::ResolvePenetration(const FVector& ProposedAdjustme
 			*ActorOwner->GetName(),
 			*UpdatedComponent->GetName(),
 			*UpdatedComponent->GetComponentLocation().ToString(),
-			*GetNameSafe(Hit.GetActor()),
+			*Hit.GetHitObjectHandle().GetName(),
 			*GetNameSafe(Hit.GetComponent()),
 			Hit.Component.IsValid() ? *Hit.GetComponent()->GetComponentLocation().ToString() : TEXT("<unknown>"),
 			Hit.PenetrationDepth,
diff --git a/Engine/Plugins/VirtualProduction/VCamCore/Source/VCamCore/Private/VCamOutputProvider.cpp b/Engine/Plugins/VirtualProduction/VCamCore/Source/VCamCore/Private/VCamOutputProvider.cpp
--- a/Engine/Plugins/VirtualProduction/VCamCore/Source/VCamCore/Private/VCamOutputProvider.cpp
+++ b/Engine/Plugins/VirtualProduction/VCamCore/Source/VCamCore/Private/VCamOutputProvider.cpp
@@ -423,7 +423,7 @@ void UVCamOutputProvider::OnTouchEventOutsideUMG(const FVector2D& InViewportPosi
 				else
 #endif
 				{
-					LastViewportTouchResult.GetActor()->NotifyActorOnClicked();
+					LastViewportTouchResult.GetHitObjectHandle().FetchActor()->NotifyActorOnClicked();
 				}
 			}
 		}
diff --git a/Engine/Source/Runtime/Engine/Classes/GameFramework/LightWeightInstanceManager.h b/Engine/Source/Runtime/Engine/Classes/GameFramework/LightWeightInstanceManager.h
--- a/Engine/Source/Runtime/Engine/Classes/GameFramework/LightWeightInstanceManager.h
+++ b/Engine/Source/Runtime/Engine/Classes/GameFramework/LightWeightInstanceManager.h
@@ -131,10 +131,11 @@ protected:
 	// handy way to check indices quickly so we don't need to iterate through the free indices list
 	UPROPERTY(Replicated)
 	TArray<bool> ValidIndices;
-
+#if WITH_EDITOR
 	//
 	// Editor functions
 	//
 protected:
 	virtual void PostEditChangeProperty(struct FPropertyChangedEvent& PropertyChangedEvent) override;	
+#endif
 };
diff --git a/Engine/Source/Runtime/Engine/Private/LightWeightInstanceManager.cpp b/Engine/Source/Runtime/Engine/Private/LightWeightInstanceManager.cpp
--- a/Engine/Source/Runtime/Engine/Private/LightWeightInstanceManager.cpp
+++ b/Engine/Source/Runtime/Engine/Private/LightWeightInstanceManager.cpp
@@ -220,7 +220,7 @@ void ALightWeightInstanceManager::OnRep_Transforms()
 {
 	// do nothing
 }
-
+#if WITH_EDITOR
 void ALightWeightInstanceManager::PostEditChangeProperty(struct FPropertyChangedEvent& PropertyChangedEvent)
 {
 	Super::PostEditChangeProperty(PropertyChangedEvent);
@@ -231,4 +231,5 @@ void ALightWeightInstanceManager::PostEditChangeProperty(struct FPropertyChanged
 	{
 		SetRepresentedClass(RepresentedClass);
 	}
-}
\ No newline at end of file
+}
+#endif
\ No newline at end of file
