Remove log

diff --git a/Engine/Plugins/Runtime/NetworkPredictionExtras/Source/NetworkPredictionExtras/Private/MockRootMotionSourceObject.cpp b/Engine/Plugins/Runtime/NetworkPredictionExtras/Source/NetworkPredictionExtras/Private/MockRootMotionSourceObject.cpp
--- a/Engine/Plugins/Runtime/NetworkPredictionExtras/Source/NetworkPredictionExtras/Private/MockRootMotionSourceObject.cpp
+++ b/Engine/Plugins/Runtime/NetworkPredictionExtras/Source/NetworkPredictionExtras/Private/MockRootMotionSourceObject.cpp
@@ -91,7 +91,7 @@ void UMockRootMotionSourceClassMap::BuildClassMap_Internal()
 
 	Lib->LoadBlueprintAssetDataFromPath(TEXT("/Game"));
 	Lib->GetAssetDataList(AssetDataList);
-	UE_LOG(LogTemp, Warning, TEXT("Found %d items in AssetDataList"), AssetDataList.Num());
+	//UE_LOG(LogTemp, Warning, TEXT("Found %d items in AssetDataList"), AssetDataList.Num());
 
 	AssetDataList.Reset();
 	Lib->LoadBlueprintAssetDataFromPath(TEXT("/NetworkPredictionExtras"));
