Network Prediction Extras
-Edigrate source changes over from UE4

#rnx

diff --git a/Engine/Plugins/Runtime/NetworkPredictionExtras/Source/NetworkPredictionExtras/Private/MockNetworkPhysics.cpp b/Engine/Plugins/Runtime/NetworkPredictionExtras/Source/NetworkPredictionExtras/Private/MockNetworkPhysics.cpp
deleted file mode 100644
diff --git a/Engine/Plugins/Runtime/NetworkPredictionExtras/Source/NetworkPredictionExtras/Private/PhysicsMovement.cpp b/Engine/Plugins/Runtime/NetworkPredictionExtras/Source/NetworkPredictionExtras/Private/PhysicsMovement.cpp
new file mode 100644
--- /dev/null
+++ b/Engine/Plugins/Runtime/NetworkPredictionExtras/Source/NetworkPredictionExtras/Private/PhysicsMovement.cpp
@@ -0,0 +1,450 @@
+// Copyright Epic Games, Inc. All Rights Reserved.
+
+#include "PhysicsMovement.h"
+#include "Engine/World.h"
+#include "Physics/Experimental/PhysScene_Chaos.h"
+#include "PBDRigidsSolver.h"
+#include "ChaosSolversModule.h"
+#include "RewindData.h"
+#include "EngineUtils.h"
+#include "Net/UnrealNetwork.h"
+#include "Chaos/SimCallbackObject.h"
+#include "UObject/UObjectIterator.h"
+#include "GameFramework/PlayerController.h"
+#include "DrawDebugHelpers.h"
+#include "PhysicsProxy/SingleParticlePhysicsProxy.h"
+#include "TimerManager.h"
+#include "NetworkPredictionDebug.h"
+#include "NetworkPredictionCVars.h"
+#include "Async/NetworkPredictionAsyncModelDef.h"
+#include "Async/NetworkPredictionAsyncWorldManager.h"
+#include "Async/NetworkPredictionAsyncID.h"
+#include "Async/NetworkPredictionAsyncModelDefRegistry.h"
+#include "Async/NetworkPredictionAsyncProxyImpl.h"
+#include "Components/PrimitiveComponent.h"
+
+
+struct FPhysicsMovementSimulation
+{
+	static void Tick_Internal(FNetworkPredictionAsyncTickContext& Context, UE_NP::FNetworkPredictionAsyncID ID, FPhysicsInputCmd& InputCmd, FPhysicsMovementLocalState& LocalState, FPhysicsMovementNetState& NetState)
+	{
+		UWorld* World = Context.World;
+		const float DeltaSeconds = Context.DeltaTimeSeconds;
+		const int32 SimulationFrame = Context.SimulationFrame;
+
+		if (!ensureAlways(LocalState.Proxy))
+		{
+			return;
+		}
+
+		auto* PT = LocalState.Proxy->GetPhysicsThreadAPI();
+		if (!PT)
+		{
+			return;
+		}
+
+		if (!InputCmd.bLegit)
+		{
+			// THis is just debug warning trying to make sure we never submit uninitialized input cmds into the system.
+			UE_LOG(LogNetworkPrediction, Warning, TEXT("Illegitimate InputCmd has seeped in!!! SimFrame: %d. LocalFrame: %d"), SimulationFrame, Context.LocalStorageFrame);
+		}
+
+		//UE_LOG(LogTemp, Warning, TEXT("[%s][%d][%d] Yaw: %f %s"), World->GetNetMode() == NM_Client ? TEXT("C") : TEXT("S"), SimulationFrame, SimulationFrame - Context.LocalStorageFrame, InputCmd.TargetYaw, Context.bIsResim ? TEXT("RESIM") : TEXT(" "));
+		//UE_LOG(LogTemp, Warning, TEXT("[%s][%d][%d] 0x%X Count: %d %s"), World->GetNetMode() == NM_Client ? TEXT("C") : TEXT("S"), SimulationFrame, SimulationFrame - Context.LocalStorageFrame, (int64)&InputCmd, InputCmd.Counter, Context.bIsResim ? TEXT("RESIM") : TEXT(" "));
+
+		//UE_LOG(LogTemp, Warning, TEXT("0x%X [%s][%d][%d][%d] ID: %d. CheckSum: %d %s"), (int64)World, World->GetNetMode() == NM_Client ? TEXT("C") : TEXT("S"), SimulationFrame, Context.LocalStorageFrame, SimulationFrame - Context.LocalStorageFrame, (int32)ID, NetState.CheckSum, Context.bIsResim ? TEXT("RESIM") : TEXT(" "));
+		NetState.CheckSum++;
+
+		//UE_LOG(LogTemp, Warning, TEXT("[%d] R: %s W: %s"), SimulationFrame, *PT->R().ToString(), *PT->W().ToString());
+
+		UE_NETWORK_PHYSICS::ConditionalFrameEnsure();
+		if (UE_NETWORK_PHYSICS::ConditionalFrameBreakpoint())
+		{
+			UE_LOG(LogTemp, Warning, TEXT("[%d][%d] Location: %s"), UE_NETWORK_PHYSICS::DebugSimulationFrame(), UE_NETWORK_PHYSICS::DebugServer(), *PT->X().ToString());
+		}
+
+		FVector TracePosition = PT->X();
+		FVector EndPosition = TracePosition + FVector(0.f, 0.f, -100.f);
+		FCollisionShape Shape = FCollisionShape::MakeSphere(250.f);
+		ECollisionChannel CollisionChannel = ECollisionChannel::ECC_WorldStatic;
+		FCollisionResponseParams ResponseParams = FCollisionResponseParams::DefaultResponseParam;
+		FCollisionObjectQueryParams ObjectParams(ECollisionChannel::ECC_PhysicsBody);
+
+		FHitResult OutHit;
+		bool bInAir = !UE_NETWORK_PHYSICS::JumpHack() && !World->LineTraceSingleByChannel(OutHit, TracePosition, EndPosition, ECollisionChannel::ECC_WorldStatic, LocalState.QueryParams, ResponseParams);
+		const float UpDot = FVector::DotProduct(PT->R().GetUpVector(), FVector::UpVector);
+
+		if (UPrimitiveComponent* HitPrim = OutHit.Component.Get())
+		{
+			if (FPhysicsInterface::IsValid(HitPrim->BodyInstance.ActorHandle) && (LocalState.Proxy == HitPrim->BodyInstance.ActorHandle))
+			{
+				ensure(false); // hit self somehow?
+			}
+		}
+
+		// Debug CVar to make jump cause mispredictions
+		if (UE_NETWORK_PHYSICS::JumpMisPredict())
+		{
+			if (InputCmd.bJumpedPressed && NetState.JumpStartFrame + 10 < SimulationFrame)
+			{
+				NetState.JumpStartFrame = SimulationFrame;
+				PT->SetLinearImpulse(FVector(0.f, 0.f, 60000.f + (FMath::FRand() * 50000.f)));
+			}
+			return;
+		}
+
+		// Impulse on physics ball (not reimplemented yet)
+		if (UE_NETWORK_PHYSICS::MockImpulse() && NetState.KickFrame + 10 < SimulationFrame)
+		{
+			/*
+			for (FSingleParticlePhysicsProxy* BallProxy : BallProxies)
+			{					
+				if (auto* BallPT = BallProxy->GetPhysicsThreadAPI())
+				{
+					const FVector BallLocation = BallPT->X();
+					const float BallRadius = BallPT->Geometry()->BoundingBox().OriginRadius(); //GetRadius();
+						
+					if (BallRadius > 0.f && (FVector::DistSquared(PT->X(), BallLocation) < BallRadius * BallRadius))
+					{
+						FVector Impulse = BallPT->X() - PT->X();
+						Impulse.Z =0.f;
+						Impulse.Normalize();
+						Impulse *= UE_NETWORK_PHYSICS::MockImpulseX();
+						Impulse.Z = UE_NETWORK_PHYSICS::MockImpulseZ();
+
+						//UE_LOG(LogTemp, Warning, TEXT("Applied Force. %s"), *GetNameSafe(HitPrimitive->GetOwner()));
+						BallPT->SetLinearImpulse( Impulse, false );
+						NetState.KickFrame = SimulationFrame;
+					}
+				}
+			}
+			*/
+		}
+
+		if (!bInAir)
+		{
+			if (NetState.InAirFrame != 0)
+			{
+				NetState.InAirFrame = 0;
+			}
+		}
+		else
+		{
+			if (NetState.InAirFrame == 0)
+			{
+				NetState.InAirFrame = SimulationFrame;
+			}
+		}
+
+		if (InputCmd.bJumpedPressed)
+		{
+			if (NetState.InAirFrame == 0 || (NetState.InAirFrame + UE_NETWORK_PHYSICS::JumpFudgeFrames() > SimulationFrame))
+			{
+				if (NetState.JumpStartFrame == 0)
+				{
+					NetState.JumpStartFrame = SimulationFrame;
+				}
+
+				if (NetState.JumpStartFrame + UE_NETWORK_PHYSICS::JumpFrameDuration() > SimulationFrame)
+				{
+					PT->AddForce( Chaos::FVec3(0.f, 0.f, UE_NETWORK_PHYSICS::JumpForce()) );
+
+					//UE_LOG(LogTemp, Warning, TEXT("[%d] Jumped [JumpStart: %d. InAir: %d]"), SimulationFrame, NetState.JumpStartFrame, NetState.InAirFrame);
+					NetState.JumpCooldownMS = 1000;
+				}
+			}
+		}
+		else
+		{
+			if (NetState.InAirFrame == 0 && (NetState.JumpStartFrame + UE_NETWORK_PHYSICS::JumpFrameDuration() < SimulationFrame))
+			{
+				NetState.JumpStartFrame = 0;
+			}
+		}
+		
+		// Keep pawn upright
+		if (NetState.bEnableKeepUpright && UpDot < 0.95f)
+		{
+			FRotator Rot = PT->R().Rotator();
+			FVector Target(0, 0, 1); // TODO choose up based on floor normal
+			FVector CurrentUp = Rot.RotateVector(Target);
+			FVector Torque = -FVector::CrossProduct(Target, CurrentUp);
+			FVector Damp(PT->W().X* UE_NETWORK_PHYSICS::TurnDampK(), PT->W().Y* UE_NETWORK_PHYSICS::TurnDampK(), 0);
+
+			PT->AddTorque(Torque* UE_NETWORK_PHYSICS::TurnK() + Damp);
+		}
+
+		if (InputCmd.bBrakesPressed)
+		{
+			FVector NewV = PT->V();
+			if (NewV.SizeSquared2D() < 1.f)
+			{
+				PT->SetV( Chaos::FVec3(0.f, 0.f, NewV.Z));
+			}
+			else
+			{
+				PT->SetV( Chaos::FVec3(NewV.X * 0.8f, NewV.Y * 0.8f, NewV.Z));
+			}
+		}
+		else
+		{
+			// Movement
+			if (InputCmd.Force.SizeSquared() > 0.001f)
+			{
+				PT->AddForce(InputCmd.Force * NetState.ForceMultiplier * UE_NETWORK_PHYSICS::MovementK());
+
+			}
+			else if(!bInAir)
+			{
+				// Auto brake: Applied when no input and grounded.
+				FVector NewV = PT->V();
+				if (NewV.SizeSquared2D() < 1.f)
+				{
+					PT->SetV(Chaos::FVec3(0.f, 0.f, NewV.Z));
+				}
+				else
+				{
+					const float DragFactor = FMath::Max(0.f, FMath::Min(1.f - (NetState.AutoBrakeStrength * DeltaSeconds), 1.f));
+					PT->SetV(Chaos::FVec3(NewV.X * DragFactor, NewV.Y * DragFactor, NewV.Z));
+				}
+			}
+
+			// Rotation
+			if (InputCmd.Torque.SizeSquared() > 0.001f)
+			{
+				PT->AddTorque(InputCmd.Torque* NetState.ForceMultiplier* UE_NETWORK_PHYSICS::RotationK());
+			}
+
+			if (NetState.bEnableAutoFaceTargetYaw)
+			{
+				// Auto Turn to target yaw
+				const float CurrentYaw = PT->R().Rotator().Yaw + (PT->W().Z * NetState.AutoFaceTargetYawDamp);
+				const float DesiredYaw = FMath::DegreesToRadians(InputCmd.TargetYaw);
+				const float DeltaYaw = FRotator::NormalizeAxis(InputCmd.TargetYaw - CurrentYaw );
+				PT->AddTorque(FVector(0.f, 0.f, DeltaYaw * NetState.AutoFaceTargetYawStrength));
+			}
+			else
+			{
+				// Prevent spheres from spinning in place.
+				FVector AngularVelocity = PT->W();
+				const float DragFactor = FMath::Max(0.f, FMath::Min(1.f - (UE_NETWORK_PHYSICS::DampYawVelocityK() * DeltaSeconds), 1.f));
+				PT->SetW(Chaos::FVec3(AngularVelocity.X, AngularVelocity.Y, AngularVelocity.Z * DragFactor));
+			}
+		}
+
+		// Drag force
+		FVector V = PT->V();
+		V.Z = 0.f;
+		if (V.SizeSquared() > 0.1f)
+		{
+			FVector Drag = -1.f * V * UE_NETWORK_PHYSICS::DragK();
+			PT->AddForce(Drag);
+		}
+		
+		// angular velocity limit
+		FVector W = PT->W();
+		{
+			const float MaxAngularVelocitySq = UE_NETWORK_PHYSICS::MaxAngularVelocity() * UE_NETWORK_PHYSICS::MaxAngularVelocity();
+			if (W.SizeSquared() > MaxAngularVelocitySq)
+			{
+				W = W.GetUnsafeNormal() * UE_NETWORK_PHYSICS::MaxAngularVelocity();
+				PT->SetW(W);
+			}
+		}
+
+		NetState.JumpCooldownMS = FMath::Max( NetState.JumpCooldownMS - (int32)(DeltaSeconds* 1000.f), 0);
+		if (NetState.JumpCooldownMS != 0)
+		{
+			UE_CLOG(UE_NETWORK_PHYSICS::MockDebug(), LogNetworkPhysics, Log, TEXT("[%d/%d] JumpCount: %d. JumpCooldown: %d"), SimulationFrame, Context.LocalStorageFrame, NetState.JumpCount, NetState.JumpCooldownMS);
+		}
+
+		if (InputCmd.bJumpedPressed)
+		{
+			// Note this is really just for debugging. "How many times was the button pressed"
+			NetState.JumpCount++;
+			UE_CLOG(UE_NETWORK_PHYSICS::MockDebug(), LogNetworkPhysics, Log, TEXT("[%d/%d] bJumpedPressed: %d. Count: %d"), SimulationFrame, Context.LocalStorageFrame, InputCmd.bJumpedPressed, NetState.JumpCount);
+		}
+	}
+};
+
+
+struct FPhysicsMovementAsyncModelDef : public UE_NP::FNetworkPredictionAsyncModelDef
+{
+	NP_ASYNC_MODEL_BODY()
+
+	using InputCmdType = FPhysicsInputCmd;
+	using NetStateType = FPhysicsMovementNetState;
+	using LocalStateType = FPhysicsMovementLocalState;
+	using SimulationTickType = FPhysicsMovementSimulation;
+
+	static const TCHAR* GetName() { return TEXT("PhysicsMovement"); }
+	static constexpr int32 GetSortPriority() { return 0; }
+};
+
+NP_ASYNC_MODEL_REGISTER(FPhysicsMovementAsyncModelDef);
+
+
+UPhysicsMovementComponent::UPhysicsMovementComponent()
+{
+	bWantsInitializeComponent = true;
+	bAutoActivate = true;
+	SetIsReplicatedByDefault(true);
+
+	PrimaryComponentTick.bCanEverTick = true;
+	PrimaryComponentTick.TickGroup = TG_DuringPhysics;
+	PrimaryComponentTick.bStartWithTickEnabled = true;
+}
+
+void UPhysicsMovementComponent::InitializeComponent()
+{
+	Super::InitializeComponent();
+#if WITH_CHAOS
+	UWorld* World = GetWorld();	
+	checkSlow(World);
+	UNetworkPhysicsManager* Manager = World->GetSubsystem<UNetworkPhysicsManager>();
+	if (!Manager)
+	{
+		return;
+	}
+
+	
+	FPhysicsMovementLocalState LocalState;
+	LocalState.Proxy = this->GetManagedProxy();
+	LocalState.QueryParams.AddIgnoredActor(GetOwner());
+	if (LocalState.Proxy)
+	{
+		if (ensure(NetworkPredictionProxy.RegisterProxy(GetWorld())))
+		{
+			NetworkPredictionProxy.RegisterSim<FPhysicsMovementAsyncModelDef>(MoveTemp(LocalState), FPhysicsMovementNetState(), &PendingInputCmd, &MovementState);
+		}
+	}
+	else
+	{
+		UE_LOG(LogNetworkPhysics, Warning, TEXT("No valid physics body found on %s"), *GetName());
+	}
+#endif
+}
+
+void UPhysicsMovementComponent::PreReplication(IRepChangedPropertyTracker & ChangedPropertyTracker)
+{
+	Super::PreReplication(ChangedPropertyTracker);
+	NetworkPredictionProxy.OnPreReplication();
+}
+
+void UPhysicsMovementComponent::UninitializeComponent()
+{
+	Super::UninitializeComponent();
+	NetworkPredictionProxy.UnregisterProxy();
+}
+
+void UPhysicsMovementComponent::TickComponent(float DeltaTime, enum ELevelTick TickType, FActorComponentTickFunction *ThisTickFunction)
+{
+	Super::TickComponent(DeltaTime, TickType, ThisTickFunction);
+	
+	APlayerController* PC = GetOwnerPC();
+
+	// RegisterController when owning PC changes, or at least once (With null controller) on server
+	if (CachedPC != PC || (!bHasRegisteredController && GetOwnerRole() == ROLE_Authority))
+	{
+		NetworkPredictionProxy.RegisterController(PC);
+		bHasRegisteredController = true;
+		CachedPC = PC;
+	}
+
+	if ((PC && PC->IsLocalController()) || (GetOwnerRole() == ROLE_Authority))
+	{
+		// Broadcast out a delegate. The user will write to PendingInputCmd
+		OnGenerateLocalInputCmd.Broadcast();
+		PendingInputCmd.bLegit = true;
+		//PendingInputCmd.Counter++;
+	}
+}
+
+void UPhysicsMovementComponent::GetLifetimeReplicatedProps(TArray<FLifetimeProperty>& OutLifetimeProps) const
+{
+	Super::GetLifetimeReplicatedProps(OutLifetimeProps);
+	DOREPLIFETIME( UPhysicsMovementComponent, NetworkPredictionProxy);
+}
+
+APlayerController* UPhysicsMovementComponent::GetOwnerPC() const
+{
+	if (APawn* PawnOwner = Cast<APawn>(GetOwner()))
+	{
+		return Cast<APlayerController>(PawnOwner->GetController());
+	}
+
+	return nullptr;
+}
+
+void UPhysicsMovementComponent::SetAutoTargetYawStrength(float Strength)
+{
+	NetworkPredictionProxy.ModifyNetState<FPhysicsMovementAsyncModelDef>([Strength](FPhysicsMovementNetState& NetState)
+	{
+		NetState.AutoFaceTargetYawStrength = Strength;
+	});	
+}
+
+void UPhysicsMovementComponent::SetAutoTargetYawDamp(float YawDamp)
+{
+	NetworkPredictionProxy.ModifyNetState<FPhysicsMovementAsyncModelDef>([YawDamp](FPhysicsMovementNetState& NetState)
+	{
+		NetState.AutoFaceTargetYawDamp = YawDamp;
+	});	
+}
+
+void UPhysicsMovementComponent::SetEnableTargetYaw(bool bTargetYaw)
+{
+	NetworkPredictionProxy.ModifyNetState<FPhysicsMovementAsyncModelDef>([bTargetYaw](FPhysicsMovementNetState& NetState)
+	{
+		NetState.bEnableAutoFaceTargetYaw = bTargetYaw;
+	});
+}
+
+
+void UPhysicsMovementComponent::SetEnableKeepUpright(bool bKeepUpright)
+{
+	NetworkPredictionProxy.ModifyNetState<FPhysicsMovementAsyncModelDef>([bKeepUpright](FPhysicsMovementNetState& NetState)
+	{
+		NetState.bEnableKeepUpright = bKeepUpright;
+	});
+}
+
+void UPhysicsMovementComponent::SetAutoBrakeStrength(float AutoBrakeStrength)
+{
+	NetworkPredictionProxy.ModifyNetState<FPhysicsMovementAsyncModelDef>([AutoBrakeStrength](FPhysicsMovementNetState& NetState)
+	{
+		NetState.AutoBrakeStrength = AutoBrakeStrength;
+	});
+}
+
+
+void UPhysicsMovementComponent::TestMisprediction()
+{
+	const int32 RandValue = FMath::RandHelper(1024);
+	UE_LOG(LogNetworkPhysics, Warning, TEXT("Setting NewRand on to %d"), RandValue);
+	NetworkPredictionProxy.ModifyNetState<FPhysicsMovementAsyncModelDef>([RandValue](FPhysicsMovementNetState& NetState)
+	{
+		NetState.RandValue = RandValue;
+	});	
+}
+
+FAutoConsoleCommandWithWorldAndArgs ForcePhysicsMovementCorrectionCmd(TEXT("np.PhysicsMovement.ForceRandCorrection2"), TEXT(""),
+FConsoleCommandWithWorldAndArgsDelegate::CreateLambda([](const TArray< FString >& Args, UWorld* InWorld) 
+	{
+		for (TObjectIterator<UWorld> It; It; ++It)
+		{
+			UWorld* World = *It;
+			if (World->GetNetMode() != NM_Client && (World->WorldType == EWorldType::PIE || World->WorldType == EWorldType::Game))
+			{
+				const int32 NewRand = FMath::RandHelper(1024);
+
+				for (TActorIterator<AActor> ActorIt(World); ActorIt; ++ActorIt)
+				{
+					if (UPhysicsMovementComponent* PhysComp = ActorIt->FindComponentByClass<UPhysicsMovementComponent>())
+					{
+						PhysComp->TestMisprediction();
+					}
+				}
+			}
+		}
+	}));
\ No newline at end of file
diff --git a/Engine/Plugins/Runtime/NetworkPredictionExtras/Source/NetworkPredictionExtras/Private/PhysicsSimple.cpp b/Engine/Plugins/Runtime/NetworkPredictionExtras/Source/NetworkPredictionExtras/Private/PhysicsSimple.cpp
new file mode 100644
--- /dev/null
+++ b/Engine/Plugins/Runtime/NetworkPredictionExtras/Source/NetworkPredictionExtras/Private/PhysicsSimple.cpp
@@ -0,0 +1,212 @@
+// Copyright Epic Games, Inc. All Rights Reserved.
+
+#include "PhysicsSimple.h"
+#include "Engine/World.h"
+#include "Physics/Experimental/PhysScene_Chaos.h"
+#include "PBDRigidsSolver.h"
+#include "ChaosSolversModule.h"
+#include "RewindData.h"
+#include "EngineUtils.h"
+#include "Net/UnrealNetwork.h"
+#include "Chaos/SimCallbackObject.h"
+#include "UObject/UObjectIterator.h"
+#include "GameFramework/PlayerController.h"
+#include "DrawDebugHelpers.h"
+#include "PhysicsProxy/SingleParticlePhysicsProxy.h"
+#include "TimerManager.h"
+#include "NetworkPredictionDebug.h"
+#include "NetworkPredictionCVars.h"
+#include "Async/NetworkPredictionAsyncModelDef.h"
+#include "Async/NetworkPredictionAsyncWorldManager.h"
+#include "Async/NetworkPredictionAsyncID.h"
+#include "Async/NetworkPredictionAsyncModelDefRegistry.h"
+#include "Async/NetworkPredictionAsyncProxyImpl.h"
+
+namespace UE_NP
+{
+	NETSIM_DEVCVAR_SHIPCONST_FLOAT(SimpleSpeed, 500.f, "np2.SimpleSpeed", "");
+}
+
+struct FSimpletSimulation
+{
+	static void Tick_Internal(FNetworkPredictionAsyncTickContext& Context, UE_NP::FNetworkPredictionAsyncID ID, FSimpleInputCmd& InputCmd, FSimpleLocalState& LocalState, FSimpleNetState& NetState)
+	{
+		UWorld* World = Context.World;
+		const float DeltaSeconds = Context.DeltaTimeSeconds;
+		const int32 SimulationFrame = Context.SimulationFrame;
+
+		if (!ensureAlways(LocalState.Proxy))
+		{
+			return;
+		}
+
+		auto* PT = LocalState.Proxy->GetPhysicsThreadAPI();
+		if (!PT)
+		{
+			return;
+		}
+
+		if (!InputCmd.bLegit)
+		{
+			// THis is just debug warning trying to make sure we never submit uninitialized input cmds into the system.
+			UE_LOG(LogNetworkPrediction, Warning, TEXT("Illegitimate InputCmd has seeped in!!! SimFrame: %d. LocalFrame: %d"), SimulationFrame, Context.LocalStorageFrame);
+		}
+
+		if (InputCmd.bButtonPressed)
+		{
+			//UE_LOG(LogTemp, Warning, TEXT("0x%X [%d]"), (int64)
+			//UE_LOG(LogTemp, Warning, TEXT("0x%X [%s][%d][%d][%d] ID: %d. [%s][%s] %s"), (int64)World, World->GetNetMode() == NM_Client ? TEXT("C") : TEXT("S"), SimulationFrame, Context.LocalStorageFrame, SimulationFrame - Context.LocalStorageFrame, (int32)ID, *InputCmd.ToString(), *NetState.ToString(), Context.bIsResim ? TEXT("RESIM") : TEXT(" "));
+			NetState.ButtonPressedCounter++;
+		}
+
+		NetState.Counter++;
+
+		if (InputCmd.MovementDir.SizeSquared() < 0.1f)
+		{
+			PT->SetV(Chaos::FVec3(0,0,0));
+		}
+		else
+		{
+			//PT->AddForce(InputCmd.MovementDir * UE_NP::SimpleSpeed());
+			PT->SetV(InputCmd.MovementDir * UE_NP::SimpleSpeed());
+		}
+	}
+};
+
+
+struct FSimpleAsyncModelDef : public UE_NP::FNetworkPredictionAsyncModelDef
+{
+	NP_ASYNC_MODEL_BODY()
+
+	using InputCmdType = FSimpleInputCmd;
+	using NetStateType = FSimpleNetState;
+	using LocalStateType = FSimpleLocalState;
+	using SimulationTickType = FSimpletSimulation;
+
+	static const TCHAR* GetName() { return TEXT("SimpleSimulation"); }
+	static constexpr int32 GetSortPriority() { return 0; }
+};
+
+NP_ASYNC_MODEL_REGISTER(FSimpleAsyncModelDef);
+
+
+UPhysicsSimpleComponent::UPhysicsSimpleComponent()
+{
+	bWantsInitializeComponent = true;
+	bAutoActivate = true;
+	SetIsReplicatedByDefault(true);
+
+	PrimaryComponentTick.bCanEverTick = true;
+	PrimaryComponentTick.TickGroup = TG_DuringPhysics;
+	PrimaryComponentTick.bStartWithTickEnabled = true;
+}
+
+void UPhysicsSimpleComponent::InitializeComponent()
+{
+	Super::InitializeComponent();
+#if WITH_CHAOS
+	UWorld* World = GetWorld();	
+	checkSlow(World);
+	UNetworkPhysicsManager* Manager = World->GetSubsystem<UNetworkPhysicsManager>();
+	if (!Manager)
+	{
+		return;
+	}
+	
+	FSimpleLocalState LocalState;
+	LocalState.Proxy = this->GetManagedProxy();
+	if (LocalState.Proxy)
+	{
+		if (ensure(NetworkPredictionProxy.RegisterProxy(GetWorld())))
+		{
+			NetworkPredictionProxy.RegisterSim<FSimpleAsyncModelDef>(MoveTemp(LocalState), FSimpleNetState(), &PendingInputCmd, &SimpleState);
+		}
+	}
+	else
+	{
+		UE_LOG(LogNetworkPhysics, Warning, TEXT("No valid physics body found on %s"), *GetName());
+	}
+#endif
+}
+
+void UPhysicsSimpleComponent::PreReplication(IRepChangedPropertyTracker & ChangedPropertyTracker)
+{
+	Super::PreReplication(ChangedPropertyTracker);
+	NetworkPredictionProxy.OnPreReplication();
+}
+
+void UPhysicsSimpleComponent::UninitializeComponent()
+{
+	Super::UninitializeComponent();
+	NetworkPredictionProxy.UnregisterProxy();
+}
+
+void UPhysicsSimpleComponent::TickComponent(float DeltaTime, enum ELevelTick TickType, FActorComponentTickFunction *ThisTickFunction)
+{
+	Super::TickComponent(DeltaTime, TickType, ThisTickFunction);
+	
+	APlayerController* PC = GetOwnerPC();
+
+	// RegisterController when owning PC changes, or at least once (With null controller) on server
+	if (CachedPC != PC || (!bHasRegisteredController && GetOwnerRole() == ROLE_Authority))
+	{
+		NetworkPredictionProxy.RegisterController(PC);
+		bHasRegisteredController = true;
+		CachedPC = PC;
+	}
+
+	if ((PC && PC->IsLocalController()) || (GetOwnerRole() == ROLE_Authority))
+	{
+		// Broadcast out a delegate. The user will write to PendingInputCmd
+		OnGenerateLocalInputCmd.Broadcast();
+		PendingInputCmd.bLegit = true;
+		//PendingInputCmd.Counter++;
+	}
+}
+
+void UPhysicsSimpleComponent::GetLifetimeReplicatedProps(TArray<FLifetimeProperty>& OutLifetimeProps) const
+{
+	Super::GetLifetimeReplicatedProps(OutLifetimeProps);
+	DOREPLIFETIME( UPhysicsSimpleComponent, NetworkPredictionProxy);
+}
+
+APlayerController* UPhysicsSimpleComponent::GetOwnerPC() const
+{
+	if (APawn* PawnOwner = Cast<APawn>(GetOwner()))
+	{
+		return Cast<APlayerController>(PawnOwner->GetController());
+	}
+
+	return nullptr;
+}
+
+void UPhysicsSimpleComponent::SetCounter(int32 NewValue)
+{
+	NetworkPredictionProxy.ModifyNetState<FSimpleAsyncModelDef>([NewValue](FSimpleNetState& NetState)
+	{
+		NetState.ButtonPressedCounter += NewValue;
+		NetState.Counter += NewValue;
+	});
+}
+
+
+FAutoConsoleCommandWithWorldAndArgs ForcePhysicsSimpleCorrectionCmd(TEXT("np.PhysicsSimple.ForceRandCorrection2"), TEXT(""),
+FConsoleCommandWithWorldAndArgsDelegate::CreateLambda([](const TArray< FString >& Args, UWorld* InWorld) 
+{
+	for (TObjectIterator<UWorld> It; It; ++It)
+	{
+		UWorld* World = *It;
+		if (World->GetNetMode() != NM_Client && (World->WorldType == EWorldType::PIE || World->WorldType == EWorldType::Game))
+		{
+			const int32 NewRand = FMath::RandHelper(1024);
+
+			for (TActorIterator<AActor> ActorIt(World); ActorIt; ++ActorIt)
+			{
+				if (UPhysicsSimpleComponent* PhysComp = ActorIt->FindComponentByClass<UPhysicsSimpleComponent>())
+				{
+					PhysComp->SetCounter(NewRand);
+				}
+			}
+		}
+	}
+}));
\ No newline at end of file
diff --git a/Engine/Plugins/Runtime/NetworkPredictionExtras/Source/NetworkPredictionExtras/Public/MockNetworkPhysics.h b/Engine/Plugins/Runtime/NetworkPredictionExtras/Source/NetworkPredictionExtras/Public/MockNetworkPhysics.h
deleted file mode 100644
diff --git a/Engine/Plugins/Runtime/NetworkPredictionExtras/Source/NetworkPredictionExtras/Public/PhysicsMovement.h b/Engine/Plugins/Runtime/NetworkPredictionExtras/Source/NetworkPredictionExtras/Public/PhysicsMovement.h
new file mode 100644
--- /dev/null
+++ b/Engine/Plugins/Runtime/NetworkPredictionExtras/Source/NetworkPredictionExtras/Public/PhysicsMovement.h
@@ -0,0 +1,253 @@
+// Copyright Epic Games, Inc. All Rights Reserved.
+
+#pragma once
+
+#include "NetworkPhysics.h"
+#include "NetworkPredictionLog.h"
+#include "Async/NetworkPredictionAsyncProxy.h"
+#include "NetworkPhysicsComponent.h"
+#include "NetworkPredictionCVars.h"
+#include "PhysicsMovement.generated.h"
+
+namespace UE_NP
+{
+	NETSIM_DEVCVAR_SHIPCONST_INT(ForceReconcilePhysicsInputCmd, 0, "np2.PhysicsInputCmdForceReconcile", "");
+	NETSIM_DEVCVAR_SHIPCONST_INT(ForceReconcilePhysicsMovementState, 0, "np2.ForceReconcilePhysicsMovementState", "");
+};
+
+USTRUCT(BlueprintType)
+struct FPhysicsInputCmd
+{
+	GENERATED_BODY()
+
+	FPhysicsInputCmd()
+		: Force(ForceInitToZero)
+		, Torque(ForceInitToZero)
+	{ }
+
+	// Simple world vector force to be applied
+	UPROPERTY(BlueprintReadWrite,Category="Input")
+	FVector	Force;
+
+	UPROPERTY(BlueprintReadWrite, Category = "Input")
+	FVector	Torque;
+
+	/** Target yaw of character (Degrees). Torque will be applied to rotate character towards target. */
+	UPROPERTY(BlueprintReadWrite, Category = "Input")
+	float TargetYaw = 0.f;
+
+	UPROPERTY(BlueprintReadWrite,Category="Input")
+	bool bJumpedPressed = false;
+
+	UPROPERTY(BlueprintReadWrite,Category="Input")
+	bool bBrakesPressed = false;
+
+	// For testing: to ensure the system only sees user created InputCmds and not defaults
+	bool bLegit = false;
+	// Debugging. 
+	int32 Counter = 0;
+
+	void NetSerialize(FArchive& Ar)
+	{
+		if (Ar.IsSaving() && !bLegit)
+		{
+			UE_LOG(LogNetworkPrediction, Warning, TEXT("FPhysicsInputCmd Sending Bad Input!"));
+		}
+
+		Ar << Force;
+		Ar << Torque;
+		Ar << TargetYaw;
+		Ar << bJumpedPressed;
+		Ar << bBrakesPressed;
+		Ar << bLegit;
+		Ar << Counter;
+
+		if (!Ar.IsSaving() && !bLegit)
+		{
+			UE_LOG(LogNetworkPrediction, Warning, TEXT("FPhysicsInputCmd Received Bad Input!"));
+		}
+	}
+
+	bool ShouldReconcile(const FPhysicsInputCmd& AuthState) const
+	{
+		return FVector::DistSquared(Force, AuthState.Force) > 0.1f
+			|| FVector::DistSquared(Torque, AuthState.Torque) > 0.1f
+			|| !FMath::IsNearlyEqual(TargetYaw, AuthState.TargetYaw, 1.0f)
+			|| bJumpedPressed != AuthState.bJumpedPressed
+			//|| Counter != AuthState.Counter // this will cause constant corrections with multiple clients but useful in testing single client
+			|| bBrakesPressed != AuthState.bBrakesPressed
+			|| (UE_NP::ForceReconcilePhysicsInputCmd() > 0);
+			
+	}
+
+	FString ToString() const
+	{
+		return FString(TEXT(""));
+	}
+};
+
+USTRUCT(BlueprintType)
+struct FPhysicsMovementNetState
+{
+	GENERATED_BODY()
+
+	// Actually used by AsyncTick to scale force applied
+	UPROPERTY(BlueprintReadWrite, Category="Mock Object")
+	float ForceMultiplier = 125000.f;
+
+	UPROPERTY(BlueprintReadWrite, Category = "Mock Object")
+	float AutoFaceTargetYawStrength = 200000.f;
+	
+	UPROPERTY(BlueprintReadWrite, Category = "Mock Object")
+	float AutoFaceTargetYawDamp = 20.f;
+
+	// If enabled, input cmd should specify target yaw.
+	UPROPERTY(BlueprintReadWrite, Category = "Mock Object")
+	bool bEnableAutoFaceTargetYaw = false;
+
+	// If enabled, object will attempt to keep its up vector aligned with world up.
+	UPROPERTY(BlueprintReadWrite, Category = "Mock Object")
+	bool bEnableKeepUpright = true;
+
+	// Strength of auto brake force applied when no input force and on ground.
+	UPROPERTY(BlueprintReadWrite, Category = "Mock Object")
+	float AutoBrakeStrength = 5.f;
+
+	// Arbitrary data that doesn't affect sim but could still trigger rollback
+	UPROPERTY(BlueprintReadWrite, Category="Mock Object")
+	int32 RandValue = 0;
+
+	UPROPERTY(BlueprintReadWrite, Category="Mock Object")
+	int32 JumpCooldownMS = 0;
+
+	// Number of frames jump has been pressed
+	UPROPERTY(BlueprintReadWrite, Category="Mock Object")
+	int32 JumpCount = 0;
+
+	UPROPERTY(BlueprintReadWrite, Category="Mock Object")
+	int32 CheckSum = 0;
+
+	// Frame we started "in air recovery" on
+	UPROPERTY(BlueprintReadWrite, Category="Mock Object")
+	int32 RecoveryFrame = 0;
+
+	// Frame we started jumping on
+	UPROPERTY(BlueprintReadWrite, Category="Mock Object")
+	int32 JumpStartFrame = 0;
+
+	// Frame we started being in the air
+	UPROPERTY(BlueprintReadWrite, Category="Mock Object")
+	int32 InAirFrame = 0;
+
+	// Frame we last applied a kick impulse
+	UPROPERTY(BlueprintReadWrite, Category="Mock Object")
+	int32 KickFrame = 0;
+
+	void NetSerialize(FArchive& Ar)
+	{
+		Ar << ForceMultiplier;
+		Ar << RandValue;
+		Ar << AutoFaceTargetYawStrength;
+		Ar << AutoFaceTargetYawDamp;
+		Ar << bEnableAutoFaceTargetYaw;
+		Ar << AutoBrakeStrength;
+		Ar << bEnableKeepUpright;
+		Ar << JumpCooldownMS;
+		Ar << JumpCount;
+		Ar << CheckSum;	
+		Ar << RecoveryFrame;
+		Ar << JumpStartFrame;
+		Ar << InAirFrame;
+		Ar << KickFrame;
+	}
+
+	bool ShouldReconcile(const FPhysicsMovementNetState& AuthState) const
+	{
+		return 
+			ForceMultiplier != AuthState.ForceMultiplier || 
+			RandValue != AuthState.RandValue ||
+			AutoFaceTargetYawStrength != AuthState.AutoFaceTargetYawStrength ||
+			AutoFaceTargetYawDamp != AuthState.AutoFaceTargetYawDamp ||
+			bEnableAutoFaceTargetYaw != AuthState.bEnableAutoFaceTargetYaw ||
+			AutoBrakeStrength != AuthState.AutoBrakeStrength ||
+			bEnableKeepUpright != AuthState.bEnableKeepUpright ||
+			JumpCooldownMS != AuthState.JumpCooldownMS || 
+			JumpCount != AuthState.JumpCount ||
+			RecoveryFrame != AuthState.RecoveryFrame ||
+			JumpStartFrame != AuthState.JumpStartFrame ||
+			InAirFrame != AuthState.InAirFrame ||
+			KickFrame!= AuthState.KickFrame ||
+			CheckSum != AuthState.CheckSum ||
+			(UE_NP::ForceReconcilePhysicsMovementState() > 0);
+	}
+};
+
+USTRUCT(BlueprintType)
+struct FPhysicsMovementLocalState
+{
+	GENERATED_BODY()
+
+	FSingleParticlePhysicsProxy* Proxy = nullptr;
+
+	FCollisionQueryParams QueryParams = FCollisionQueryParams::DefaultQueryParam;
+};
+
+UCLASS(BlueprintType, meta=(BlueprintSpawnableComponent))
+class NETWORKPREDICTIONEXTRAS_API UPhysicsMovementComponent : public UNetworkPhysicsComponent
+{
+	GENERATED_BODY()
+
+public:
+
+	UPhysicsMovementComponent();
+
+	virtual void InitializeComponent() override;
+	virtual void UninitializeComponent() override;
+	virtual void TickComponent(float DeltaTime, enum ELevelTick TickType, FActorComponentTickFunction *ThisTickFunction) override;
+	virtual void PreReplication(IRepChangedPropertyTracker & ChangedPropertyTracker);
+
+	UPROPERTY(Replicated)
+	FNetworkPredictionAsyncProxy NetworkPredictionProxy;
+
+	// Latest MovementState
+	UPROPERTY(BlueprintReadOnly, Category="Movement")
+	FPhysicsMovementNetState MovementState;
+
+	// Game code should write to this when locally controlling this object.
+	UPROPERTY(BlueprintReadWrite, Category="Input")
+	FPhysicsInputCmd PendingInputCmd;
+
+	DECLARE_DYNAMIC_MULTICAST_DELEGATE(FOnGenerateLocalInputCmd);
+
+	// Event called when local input is needed: fill out PendingInputCmd with latest input data
+	UPROPERTY(BlueprintAssignable, Category = "Input")
+	FOnGenerateLocalInputCmd OnGenerateLocalInputCmd;
+
+	
+
+	UFUNCTION(BlueprintCallable, Category="Movement")
+	void SetAutoTargetYawStrength(float Strength);
+
+	UFUNCTION(BlueprintCallable, Category="Movement")
+	void SetAutoTargetYawDamp(float YawDamp);
+
+	UFUNCTION(BlueprintCallable, Category="Movement")
+	void SetEnableTargetYaw(bool bTargetYaw);
+
+	UFUNCTION(BlueprintCallable, Category="Movement")
+	void SetEnableKeepUpright(bool bKeepUpright);
+
+	UFUNCTION(BlueprintCallable, Category="Movement")
+	void SetAutoBrakeStrength(float BrakeStrength);
+
+	void TestMisprediction();
+
+private:
+
+	APlayerController* GetOwnerPC() const;
+
+	// Crude way of detecting possession change
+	UPROPERTY()
+	APlayerController* CachedPC = nullptr;
+	bool bHasRegisteredController = false;
+};
\ No newline at end of file
diff --git a/Engine/Plugins/Runtime/NetworkPredictionExtras/Source/NetworkPredictionExtras/Public/PhysicsSimple.h b/Engine/Plugins/Runtime/NetworkPredictionExtras/Source/NetworkPredictionExtras/Public/PhysicsSimple.h
new file mode 100644
--- /dev/null
+++ b/Engine/Plugins/Runtime/NetworkPredictionExtras/Source/NetworkPredictionExtras/Public/PhysicsSimple.h
@@ -0,0 +1,152 @@
+// Copyright Epic Games, Inc. All Rights Reserved.
+
+#pragma once
+
+#include "NetworkPhysics.h"
+#include "NetworkPredictionLog.h"
+#include "Async/NetworkPredictionAsyncProxy.h"
+#include "NetworkPhysicsComponent.h"
+#include "NetworkPredictionCVars.h"
+#include "PhysicsSimple.generated.h"
+
+namespace UE_NP
+{
+	NETSIM_DEVCVAR_SHIPCONST_INT(ForceReconcileSimpleInput, 0, "np2.ForceReconcileSimpleInput", "");
+	NETSIM_DEVCVAR_SHIPCONST_INT(ForceReconcileSimpleState, 0, "np2.ForceReconcileSimpleState", "");
+};
+
+USTRUCT(BlueprintType)
+struct FSimpleInputCmd
+{
+	GENERATED_BODY()
+
+	FSimpleInputCmd()
+		: MovementDir(ForceInitToZero)
+	{ }
+
+	// Simple world vector force to be applied
+	UPROPERTY(BlueprintReadWrite,Category="Input")
+	FVector	MovementDir;
+
+	UPROPERTY(BlueprintReadWrite,Category="Input")
+	bool	bButtonPressed;
+	
+	UPROPERTY(BlueprintReadWrite,Category="Input")
+	bool	bLegit = false;
+
+	void NetSerialize(FArchive& Ar)
+	{
+		if (Ar.IsSaving() && !bLegit)
+		{
+			UE_LOG(LogNetworkPrediction, Warning, TEXT("FSimpleInputCmd Sending Bad Input!"));
+		}
+		
+		Ar << MovementDir;
+		Ar << bButtonPressed;
+		Ar << bLegit;
+
+		if (!Ar.IsSaving() && !bLegit)
+		{
+			UE_LOG(LogNetworkPrediction, Warning, TEXT("FSimpleInputCmd Received Bad Input!"));
+		}
+	}
+
+	bool ShouldReconcile(const FSimpleInputCmd& AuthState) const
+	{
+		return	MovementDir != AuthState.MovementDir ||
+				bButtonPressed != AuthState.bButtonPressed ||
+				bLegit != AuthState.bLegit ||
+				UE_NP::ForceReconcileSimpleInput() > 0;
+	}
+
+	FString ToString() const
+	{
+		return FString::Printf(TEXT("[%d] %s %s"), bButtonPressed, *MovementDir.ToString(), bLegit ? TEXT("") : TEXT("BAD INPUT"));
+	}
+
+	bool ShouldLog() const { return bButtonPressed; }
+};
+
+USTRUCT(BlueprintType)
+struct FSimpleNetState
+{
+	GENERATED_BODY()
+
+	UPROPERTY(BlueprintReadWrite, Category="Simple Object")
+	int32 ButtonPressedCounter = 0;
+
+	// Frame we last applied a kick impulse
+	UPROPERTY(BlueprintReadWrite, Category="Simple Object")
+	int32 Counter = 0;
+
+	void NetSerialize(FArchive& Ar)
+	{
+		Ar << ButtonPressedCounter;
+		Ar << Counter;
+	}
+
+	bool ShouldReconcile(const FSimpleNetState& AuthState) const
+	{
+		return	ButtonPressedCounter != AuthState.ButtonPressedCounter ||
+				Counter != AuthState.Counter ||
+				UE_NP::ForceReconcileSimpleState() > 0;
+	}
+	
+	FString ToString() const
+	{
+		return FString::Printf(TEXT("[%d/%d]"), ButtonPressedCounter, Counter);	
+	}
+};
+
+USTRUCT(BlueprintType)
+struct FSimpleLocalState
+{
+	GENERATED_BODY()
+
+	FSingleParticlePhysicsProxy* Proxy = nullptr;
+};
+
+// Debug/Testing component
+UCLASS(BlueprintType, meta=(BlueprintSpawnableComponent))
+class NETWORKPREDICTIONEXTRAS_API UPhysicsSimpleComponent : public UNetworkPhysicsComponent
+{
+	GENERATED_BODY()
+
+public:
+
+	UPhysicsSimpleComponent();
+
+	virtual void InitializeComponent() override;
+	virtual void UninitializeComponent() override;
+	virtual void TickComponent(float DeltaTime, enum ELevelTick TickType, FActorComponentTickFunction *ThisTickFunction) override;
+	virtual void PreReplication(IRepChangedPropertyTracker & ChangedPropertyTracker);
+
+	UPROPERTY(Replicated)
+	FNetworkPredictionAsyncProxy NetworkPredictionProxy;
+
+	// Latest MovementState
+	UPROPERTY(BlueprintReadOnly, Category="Movement")
+	FSimpleNetState SimpleState;
+
+	// Game code should write to this when locally controlling this object.
+	UPROPERTY(BlueprintReadWrite, Category="Input")
+	FSimpleInputCmd PendingInputCmd;
+
+	DECLARE_DYNAMIC_MULTICAST_DELEGATE(FOnGenerateLocalInputCmd);
+
+	// Event called when local input is needed: fill out PendingInputCmd with latest input data
+	UPROPERTY(BlueprintAssignable, Category = "Input")
+	FOnGenerateLocalInputCmd OnGenerateLocalInputCmd;
+
+	UFUNCTION(BlueprintCallable, Category="Movement")
+	void SetCounter(int32 NewValue);
+
+private:
+
+	APlayerController* GetOwnerPC() const;
+
+	// Crude way of detecting possession change
+	UPROPERTY()
+	APlayerController* CachedPC = nullptr;
+	bool bHasRegisteredController = false;
+};
\ No newline at end of file
