Add Acceleration/AngularAcceleration inputs to net prediction movement so we can use movement agnostic of character mass.
Change drag to remove velocity instead of applying force to avoid being affected by mass.
Jump now scales with mass.

diff --git a/Engine/Plugins/Runtime/NetworkPrediction/Source/NetworkPrediction/Public/NetworkPredictionCVars.h b/Engine/Plugins/Runtime/NetworkPrediction/Source/NetworkPrediction/Public/NetworkPredictionCVars.h
--- a/Engine/Plugins/Runtime/NetworkPrediction/Source/NetworkPrediction/Public/NetworkPredictionCVars.h
+++ b/Engine/Plugins/Runtime/NetworkPrediction/Source/NetworkPrediction/Public/NetworkPredictionCVars.h
@@ -111,7 +111,7 @@ namespace UE_NETWORK_PHYSICS
 	NETSIM_DEVCVAR_SHIPCONST_INT(bInputDecay, 0, "np2.InputDecay", "Enable Input Decay Feature");		
 	NETSIM_DEVCVAR_SHIPCONST_FLOAT(InputDecayRate, 0.99f, "np2.InputDecayRate", "Rate of input decay");
 
-	NETSIM_DEVCVAR_SHIPCONST_FLOAT(DragK, 200.f, "np2.Mock.DragK", "Drag Coefficient (higher=more drag)");
+	NETSIM_DEVCVAR_SHIPCONST_FLOAT(DragK, 1.12f, "np2.Mock.DragK", "Drag Coefficient (higher=more drag)");
 	NETSIM_DEVCVAR_SHIPCONST_FLOAT(MaxAngularVelocity, 30.f, "np2.Mock.MaxAngularVelocity", "Limits how fast character can possibly rotate.")
 	NETSIM_DEVCVAR_SHIPCONST_FLOAT(MovementK, 1.25, "np2.Mock.MovementK", "Movement Coefficient (higher=faster movement)");
 	NETSIM_DEVCVAR_SHIPCONST_FLOAT(RotationK, 1.25, "np2.Mock.RotationK", "Rotation Coefficient (higher=faster movement)");
@@ -119,7 +119,7 @@ namespace UE_NETWORK_PHYSICS
 	NETSIM_DEVCVAR_SHIPCONST_FLOAT(TurnDampK, 100.f, "np2.Mock.TurnDampK", "Coefficient for damping portion of turn. Higher=more damping but too higher will lead to instability.");
 	NETSIM_DEVCVAR_SHIPCONST_FLOAT(DampYawVelocityK, 2.f, "np2.Mock.DampYawVelocityK", "Coefficient for damping angular velocity in yaw direction only. This is only enabled when auto target yaw is disabled.");
 
-	NETSIM_DEVCVAR_SHIPCONST_FLOAT(JumpForce, 1000000.0f, "np2.Mock.JumpForce", "Per-Frame force to apply while jumping.");	
+	NETSIM_DEVCVAR_SHIPCONST_FLOAT(JumpForce, 70, "np2.Mock.JumpForce", "Per-Frame force to apply while jumping.");
 	NETSIM_DEVCVAR_SHIPCONST_INT(JumpFrameDuration, 4, "np2.Mock.JumpFrameDuration", "How many frames to apply jump force for");
 	NETSIM_DEVCVAR_SHIPCONST_INT(JumpFudgeFrames, 10, "np2.Mock.JumpFudgeFrames", "How many frames after being in air do we still allow a jump to begin");
 	NETSIM_DEVCVAR_SHIPCONST_INT(JumpHack, 0, "np2.Mock.JumpHack", "Make jump not rely on trace which currently causes non determinism");
diff --git a/Engine/Plugins/Runtime/NetworkPredictionExtras/Source/NetworkPredictionExtras/Private/PhysicsMovement.cpp b/Engine/Plugins/Runtime/NetworkPredictionExtras/Source/NetworkPredictionExtras/Private/PhysicsMovement.cpp
--- a/Engine/Plugins/Runtime/NetworkPredictionExtras/Source/NetworkPredictionExtras/Private/PhysicsMovement.cpp
+++ b/Engine/Plugins/Runtime/NetworkPredictionExtras/Source/NetworkPredictionExtras/Private/PhysicsMovement.cpp
@@ -9,6 +9,7 @@
 #include "EngineUtils.h"
 #include "Net/UnrealNetwork.h"
 #include "Chaos/SimCallbackObject.h"
+#include "Chaos/Utilities.h"
 #include "UObject/UObjectIterator.h"
 #include "GameFramework/PlayerController.h"
 #include "DrawDebugHelpers.h"
@@ -147,7 +148,7 @@ struct FPhysicsMovementSimulation
 
 				if (NetState.JumpStartFrame + UE_NETWORK_PHYSICS::JumpFrameDuration() > SimulationFrame)
 				{
-					PT->AddForce( Chaos::FVec3(0.f, 0.f, UE_NETWORK_PHYSICS::JumpForce()) );
+					PT->AddForce( Chaos::FVec3(0.f, 0.f, UE_NETWORK_PHYSICS::JumpForce()) * NetState.JumpStrength * PT->M() );
 
 					//UE_LOG(LogTemp, Warning, TEXT("[%d] Jumped [JumpStart: %d. InAir: %d]"), SimulationFrame, NetState.JumpStartFrame, NetState.InAirFrame);
 					NetState.JumpCooldownMS = 1000;
@@ -189,9 +190,11 @@ struct FPhysicsMovementSimulation
 		else
 		{
 			// Movement
-			if (InputCmd.Force.SizeSquared() > 0.001f)
+			const bool bHasLinearForceInput = InputCmd.Force.SizeSquared() > 0.0f || InputCmd.Acceleration.SizeSquared() > 0.0f;
+			if (bHasLinearForceInput)
 			{
-				PT->AddForce(InputCmd.Force * NetState.ForceMultiplier * UE_NETWORK_PHYSICS::MovementK());
+				const FVector ForceFromAcceleration = InputCmd.Acceleration * PT->M();
+				PT->AddForce((InputCmd.Force + ForceFromAcceleration) * NetState.ForceMultiplier * UE_NETWORK_PHYSICS::MovementK());
 
 			}
 			else if(!bInAir)
@@ -210,9 +213,13 @@ struct FPhysicsMovementSimulation
 			}
 
 			// Rotation
-			if (InputCmd.Torque.SizeSquared() > 0.001f)
+			const bool bHasAngularForceInput = InputCmd.Torque.SizeSquared() > 0.0f || InputCmd.AngularAcceleration.SizeSquared() > 0.0f;
+			if (bHasAngularForceInput)
 			{
-				PT->AddTorque(InputCmd.Torque* NetState.ForceMultiplier* UE_NETWORK_PHYSICS::RotationK());
+				const Chaos::FMatrix33 WorldI = Chaos::Utilities::ComputeWorldSpaceInertia(PT->R() * PT->RotationOfMass(), PT->I());
+				const FVector TorqueFromAngularAcceleration = WorldI * InputCmd.AngularAcceleration;
+				const FVector TorqueInput = InputCmd.Torque * NetState.ForceMultiplier * UE_NETWORK_PHYSICS::RotationK();
+				PT->AddTorque(TorqueInput + TorqueFromAngularAcceleration);
 			}
 
 			if (NetState.bEnableAutoFaceTargetYaw)
@@ -232,13 +239,14 @@ struct FPhysicsMovementSimulation
 			}
 		}
 
-		// Drag force
-		FVector V = PT->V();
-		V.Z = 0.f;
-		if (V.SizeSquared() > 0.1f)
+		// Drag
 		{
-			FVector Drag = -1.f * V * UE_NETWORK_PHYSICS::DragK();
-			PT->AddForce(Drag);
+			FVector NewV = PT->V();
+			if (NewV.SizeSquared() > 0.1f)
+			{
+				const float DragFactor = FMath::Max(0.f, FMath::Min(1.f - (UE_NETWORK_PHYSICS::DragK() * DeltaSeconds), 1.f));
+				PT->SetV(Chaos::FVec3(NewV.X* DragFactor, NewV.Y* DragFactor, NewV.Z * DragFactor));
+			}
 		}
 		
 		// angular velocity limit
diff --git a/Engine/Plugins/Runtime/NetworkPredictionExtras/Source/NetworkPredictionExtras/Public/PhysicsMovement.h b/Engine/Plugins/Runtime/NetworkPredictionExtras/Source/NetworkPredictionExtras/Public/PhysicsMovement.h
--- a/Engine/Plugins/Runtime/NetworkPredictionExtras/Source/NetworkPredictionExtras/Public/PhysicsMovement.h
+++ b/Engine/Plugins/Runtime/NetworkPredictionExtras/Source/NetworkPredictionExtras/Public/PhysicsMovement.h
@@ -23,6 +23,8 @@ struct FPhysicsInputCmd
 	FPhysicsInputCmd()
 		: Force(ForceInitToZero)
 		, Torque(ForceInitToZero)
+		, Acceleration(ForceInitToZero)
+		, AngularAcceleration(ForceInitToZero)
 	{ }
 
 	// Simple world vector force to be applied
@@ -32,6 +34,12 @@ struct FPhysicsInputCmd
 	UPROPERTY(BlueprintReadWrite, Category = "Input")
 	FVector	Torque;
 
+	UPROPERTY(BlueprintReadWrite, Category = "Input")
+	FVector	Acceleration;
+
+	UPROPERTY(BlueprintReadWrite, Category = "Input")
+	FVector	AngularAcceleration;
+
 	/** Target yaw of character (Degrees). Torque will be applied to rotate character towards target. */
 	UPROPERTY(BlueprintReadWrite, Category = "Input")
 	float TargetYaw = 0.f;
@@ -56,6 +64,8 @@ struct FPhysicsInputCmd
 
 		Ar << Force;
 		Ar << Torque;
+		Ar << Acceleration;
+		Ar << AngularAcceleration;
 		Ar << TargetYaw;
 		Ar << bJumpedPressed;
 		Ar << bBrakesPressed;
@@ -72,6 +82,8 @@ struct FPhysicsInputCmd
 	{
 		return FVector::DistSquared(Force, AuthState.Force) > 0.1f
 			|| FVector::DistSquared(Torque, AuthState.Torque) > 0.1f
+			|| FVector::DistSquared(Acceleration, AuthState.Acceleration) > 0.1f
+			|| FVector::DistSquared(AngularAcceleration, AuthState.AngularAcceleration) > 0.1f
 			|| !FMath::IsNearlyEqual(TargetYaw, AuthState.TargetYaw, 1.0f)
 			|| bJumpedPressed != AuthState.bJumpedPressed
 			//|| Counter != AuthState.Counter // this will cause constant corrections with multiple clients but useful in testing single client
@@ -95,6 +107,9 @@ struct FPhysicsMovementNetState
 	UPROPERTY(BlueprintReadWrite, Category="Mock Object")
 	float ForceMultiplier = 125000.f;
 
+	UPROPERTY(BlueprintReadWrite, Category="Mock Object")
+	float JumpStrength = 100.f;
+
 	UPROPERTY(BlueprintReadWrite, Category = "Mock Object")
 	float AutoFaceTargetYawStrength = 200000.f;
 	
@@ -146,6 +161,7 @@ struct FPhysicsMovementNetState
 	void NetSerialize(FArchive& Ar)
 	{
 		Ar << ForceMultiplier;
+		Ar << JumpStrength;
 		Ar << RandValue;
 		Ar << AutoFaceTargetYawStrength;
 		Ar << AutoFaceTargetYawDamp;
@@ -165,6 +181,7 @@ struct FPhysicsMovementNetState
 	{
 		return 
 			ForceMultiplier != AuthState.ForceMultiplier || 
+			JumpStrength != AuthState.JumpStrength ||
 			RandValue != AuthState.RandValue ||
 			AutoFaceTargetYawStrength != AuthState.AutoFaceTargetYawStrength ||
 			AutoFaceTargetYawDamp != AuthState.AutoFaceTargetYawDamp ||
