Network Prediction
-Added native GameMode so that cheats can be used in networked testing

diff --git a/Engine/Plugins/Runtime/NetworkPredictionExtras/Source/NetworkPredictionExtras/Private/NetworkPredictionExtrasGameMode.cpp b/Engine/Plugins/Runtime/NetworkPredictionExtras/Source/NetworkPredictionExtras/Private/NetworkPredictionExtrasGameMode.cpp
new file mode 100644
--- /dev/null
+++ b/Engine/Plugins/Runtime/NetworkPredictionExtras/Source/NetworkPredictionExtras/Private/NetworkPredictionExtrasGameMode.cpp
@@ -0,0 +1,10 @@
+// Copyright Epic Games, Inc. All Rights Reserved.
+
+#include "NetworkPredictionExtrasGameMode.h"
+#include "NetworkPredictionExtrasGameState.h"
+
+ANetworkPredictionExtrasGameMode::ANetworkPredictionExtrasGameMode(const FObjectInitializer& ObjectInitializer)
+	: Super(ObjectInitializer)
+{
+	GameStateClass = ANetworkPredictionExtrasGameState::StaticClass();
+}
\ No newline at end of file
diff --git a/Engine/Plugins/Runtime/NetworkPredictionExtras/Source/NetworkPredictionExtras/Public/NetworkPredictionExtrasGameMode.h b/Engine/Plugins/Runtime/NetworkPredictionExtras/Source/NetworkPredictionExtras/Public/NetworkPredictionExtrasGameMode.h
new file mode 100644
--- /dev/null
+++ b/Engine/Plugins/Runtime/NetworkPredictionExtras/Source/NetworkPredictionExtras/Public/NetworkPredictionExtrasGameMode.h
@@ -0,0 +1,19 @@
+// Copyright Epic Games, Inc. All Rights Reserved.
+
+#pragma once
+
+#include "CoreMinimal.h"
+
+#include "GameFramework/GameMode.h"
+#include "Engine/Engine.h"
+#include "NetworkPredictionExtrasGameMode.generated.h"
+
+UCLASS(BlueprintType)
+class ANetworkPredictionExtrasGameMode : public AGameMode
+{
+	GENERATED_BODY()
+
+	ANetworkPredictionExtrasGameMode(const FObjectInitializer& ObjectInitializer = FObjectInitializer::Get());
+
+	virtual bool AllowCheats(APlayerController* P) { return true; }
+};
\ No newline at end of file
