Network Prediction Extras
-Edigrate source changes over from UE4

#rnx

diff --git a/Engine/Plugins/Runtime/NetworkPredictionExtras/Source/NetworkPredictionExtras/Private/PhysicsMovement.cpp b/Engine/Plugins/Runtime/NetworkPredictionExtras/Source/NetworkPredictionExtras/Private/PhysicsMovement.cpp
--- a/Engine/Plugins/Runtime/NetworkPredictionExtras/Source/NetworkPredictionExtras/Private/PhysicsMovement.cpp
+++ b/Engine/Plugins/Runtime/NetworkPredictionExtras/Source/NetworkPredictionExtras/Private/PhysicsMovement.cpp
@@ -21,6 +21,8 @@
 #include "Async/NetworkPredictionAsyncID.h"
 #include "Async/NetworkPredictionAsyncModelDefRegistry.h"
 #include "Async/NetworkPredictionAsyncProxyImpl.h"
+#include "Components/PrimitiveComponent.h"
+
 
 struct FPhysicsMovementSimulation
 {
@@ -53,6 +55,8 @@ struct FPhysicsMovementSimulation
 		//UE_LOG(LogTemp, Warning, TEXT("0x%X [%s][%d][%d][%d] ID: %d. CheckSum: %d %s"), (int64)World, World->GetNetMode() == NM_Client ? TEXT("C") : TEXT("S"), SimulationFrame, Context.LocalStorageFrame, SimulationFrame - Context.LocalStorageFrame, (int32)ID, NetState.CheckSum, Context.bIsResim ? TEXT("RESIM") : TEXT(" "));
 		NetState.CheckSum++;
 
+		//UE_LOG(LogTemp, Warning, TEXT("[%d] R: %s W: %s"), SimulationFrame, *PT->R().ToString(), *PT->W().ToString());
+
 		UE_NETWORK_PHYSICS::ConditionalFrameEnsure();
 		if (UE_NETWORK_PHYSICS::ConditionalFrameBreakpoint())
 		{
@@ -443,4 +447,4 @@ FConsoleCommandWithWorldAndArgsDelegate::CreateLambda([](const TArray< FString >
 				}
 			}
 		}
-	}));
\ No newline at end of file
+	}));
