Network Physics
-Added cvars for disabling various features in the physics movement code

diff --git a/Engine/Plugins/Runtime/NetworkPredictionExtras/Source/NetworkPredictionExtras/Private/PhysicsMovement.cpp b/Engine/Plugins/Runtime/NetworkPredictionExtras/Source/NetworkPredictionExtras/Private/PhysicsMovement.cpp
--- a/Engine/Plugins/Runtime/NetworkPredictionExtras/Source/NetworkPredictionExtras/Private/PhysicsMovement.cpp
+++ b/Engine/Plugins/Runtime/NetworkPredictionExtras/Source/NetworkPredictionExtras/Private/PhysicsMovement.cpp
@@ -25,6 +25,37 @@
 #include "Components/PrimitiveComponent.h"
 
 
+NETSIM_DEVCVAR_SHIPCONST_INT(DisablePhysicsMovement, 0, "np2.DisablePhysicsMovement", "");
+
+NETSIM_DEVCVAR_SHIPCONST_INT(DisableKeepUpright, 0, "np2.PhysicsMovement.DisableKeepUpright", "");
+NETSIM_DEVCVAR_SHIPCONST_INT(DisableAutoBrake, 0, "np2.PhysicsMovement.DisableAutoBrake", "");
+NETSIM_DEVCVAR_SHIPCONST_INT(DisableMovement, 0, "np2.PhysicsMovement.DisableMovement", "");
+
+NETSIM_DEVCVAR_SHIPCONST_INT(DisableAngularMovement, 0, "np2.PhysicsMovement.DisableAngularMovement", "");
+NETSIM_DEVCVAR_SHIPCONST_INT(DisableAutoYaw, 0, "np2.PhysicsMovement.DisableAutoYaw", "");
+NETSIM_DEVCVAR_SHIPCONST_INT(DisableAntiSpin, 0, "np2.PhysicsMovement.DisableAntiSpin", "");
+NETSIM_DEVCVAR_SHIPCONST_INT(DisableDrag, 0, "np2.PhysicsMovement.DisableDrag", "");
+
+NETSIM_DEVCVAR_SHIPCONST_INT(DisableAngularVelLimit, 0, "np2.PhysicsMovement.DisableAngularVelLimit", "");
+
+
+NETSIM_DEVCVAR_SHIPCONST_INT(ForceReconcilePhysicsInputCmd, 0, "np2.PhysicsInputCmdForceReconcile", "");
+NETSIM_DEVCVAR_SHIPCONST_FLOAT(TargetYawTolerance, 0.1f, "np2.PhysicsMovement.TargetYawTolerance", "");
+
+bool FPhysicsInputCmd::ShouldReconcile(const FPhysicsInputCmd& AuthState) const
+{
+	return FVector::DistSquared(Force, AuthState.Force) > 0.1f
+		|| FVector::DistSquared(Torque, AuthState.Torque) > 0.1f
+		|| FVector::DistSquared(Acceleration, AuthState.Acceleration) > 0.1f
+		|| FVector::DistSquared(AngularAcceleration, AuthState.AngularAcceleration) > 0.1f
+		|| !FMath::IsNearlyEqual(TargetYaw, AuthState.TargetYaw, TargetYawTolerance())
+		|| bJumpedPressed != AuthState.bJumpedPressed
+		//|| Counter != AuthState.Counter // this will cause constant corrections with multiple clients but useful in testing single client
+		|| bBrakesPressed != AuthState.bBrakesPressed
+		|| (ForceReconcilePhysicsInputCmd() > 0);
+			
+}
+
 struct FPhysicsMovementSimulation
 {
 	static void Tick_Internal(FNetworkPredictionAsyncTickContext& Context, UE_NP::FNetworkPredictionAsyncID ID, FPhysicsInputCmd& InputCmd, FPhysicsMovementLocalState& LocalState, FPhysicsMovementNetState& NetState)
@@ -33,6 +64,11 @@ struct FPhysicsMovementSimulation
 		const float DeltaSeconds = Context.DeltaTimeSeconds;
 		const int32 SimulationFrame = Context.SimulationFrame;
 
+		if (DisablePhysicsMovement() == 1)
+		{
+			return;
+		}
+
 		if (!ensureAlways(LocalState.Proxy))
 		{
 			return;
@@ -164,7 +200,7 @@ struct FPhysicsMovementSimulation
 		}
 		
 		// Keep pawn upright
-		if (NetState.bEnableKeepUpright && UpDot < 0.95f)
+		if (DisableKeepUpright() == 0 && NetState.bEnableKeepUpright && UpDot < 0.95f)
 		{
 			FRotator Rot = PT->R().Rotator();
 			FVector Target(0, 0, 1); // TODO choose up based on floor normal
@@ -191,13 +227,13 @@ struct FPhysicsMovementSimulation
 		{
 			// Movement
 			const bool bHasLinearForceInput = InputCmd.Force.SizeSquared() > 0.0f || InputCmd.Acceleration.SizeSquared() > 0.0f;
-			if (bHasLinearForceInput)
+			if (bHasLinearForceInput && DisableMovement() == 0)
 			{
 				const FVector ForceFromAcceleration = InputCmd.Acceleration * PT->M();
 				PT->AddForce((InputCmd.Force + ForceFromAcceleration) * NetState.ForceMultiplier * UE_NETWORK_PHYSICS::MovementK());
 
 			}
-			else if(!bInAir)
+			else if(!bInAir && DisableAutoBrake() == 0)
 			{
 				// Auto brake: Applied when no input and grounded.
 				FVector NewV = PT->V();
@@ -214,7 +250,7 @@ struct FPhysicsMovementSimulation
 
 			// Rotation
 			const bool bHasAngularForceInput = InputCmd.Torque.SizeSquared() > 0.0f || InputCmd.AngularAcceleration.SizeSquared() > 0.0f;
-			if (bHasAngularForceInput)
+			if (bHasAngularForceInput && DisableAngularMovement() == 0)
 			{
 				const Chaos::FMatrix33 WorldI = Chaos::Utilities::ComputeWorldSpaceInertia(PT->R() * PT->RotationOfMass(), PT->I());
 				const FVector TorqueFromAngularAcceleration = WorldI * InputCmd.AngularAcceleration;
@@ -222,7 +258,7 @@ struct FPhysicsMovementSimulation
 				PT->AddTorque(TorqueInput + TorqueFromAngularAcceleration);
 			}
 
-			if (NetState.bEnableAutoFaceTargetYaw)
+			if (NetState.bEnableAutoFaceTargetYaw && DisableAutoYaw() == 0)
 			{
 				// Auto Turn to target yaw
 				const float CurrentYaw = PT->R().Rotator().Yaw + (PT->W().Z * NetState.AutoFaceTargetYawDamp);
@@ -232,17 +268,20 @@ struct FPhysicsMovementSimulation
 			}
 			else
 			{
-				// Prevent spheres from spinning in place.
-				FVector AngularVelocity = PT->W();
-				const float DragFactor = FMath::Max(0.f, FMath::Min(1.f - (UE_NETWORK_PHYSICS::DampYawVelocityK() * DeltaSeconds), 1.f));
-				PT->SetW(Chaos::FVec3(AngularVelocity.X, AngularVelocity.Y, AngularVelocity.Z * DragFactor));
+				if (DisableAntiSpin() == 0)
+				{
+					// Prevent spheres from spinning in place.
+					FVector AngularVelocity = PT->W();
+					const float DragFactor = FMath::Max(0.f, FMath::Min(1.f - (UE_NETWORK_PHYSICS::DampYawVelocityK() * DeltaSeconds), 1.f));
+					PT->SetW(Chaos::FVec3(AngularVelocity.X, AngularVelocity.Y, AngularVelocity.Z * DragFactor));
+				}
 			}
 		}
 
 		// Drag
 		{
 			FVector NewV = PT->V();
-			if (NewV.SizeSquared() > 0.1f)
+			if (NewV.SizeSquared() > 0.1f && DisableDrag() == 0)
 			{
 				const float DragFactor = FMath::Max(0.f, FMath::Min(1.f - (UE_NETWORK_PHYSICS::DragK() * DeltaSeconds), 1.f));
 				PT->SetV(Chaos::FVec3(NewV.X* DragFactor, NewV.Y* DragFactor, NewV.Z * DragFactor));
@@ -253,7 +292,7 @@ struct FPhysicsMovementSimulation
 		FVector W = PT->W();
 		{
 			const float MaxAngularVelocitySq = UE_NETWORK_PHYSICS::MaxAngularVelocity() * UE_NETWORK_PHYSICS::MaxAngularVelocity();
-			if (W.SizeSquared() > MaxAngularVelocitySq)
+			if (W.SizeSquared() > MaxAngularVelocitySq && DisableAngularVelLimit() == 0)
 			{
 				W = W.GetUnsafeNormal() * UE_NETWORK_PHYSICS::MaxAngularVelocity();
 				PT->SetW(W);
@@ -309,12 +348,6 @@ void UPhysicsMovementComponent::InitializeComponent()
 #if WITH_CHAOS
 	UWorld* World = GetWorld();	
 	checkSlow(World);
-	UNetworkPhysicsManager* Manager = World->GetSubsystem<UNetworkPhysicsManager>();
-	if (!Manager)
-	{
-		return;
-	}
-
 	
 	FPhysicsMovementLocalState LocalState;
 	LocalState.Proxy = this->GetManagedProxy();
@@ -364,7 +397,7 @@ void UPhysicsMovementComponent::TickComponent(float DeltaTime, enum ELevelTick T
 		// Broadcast out a delegate. The user will write to PendingInputCmd
 		OnGenerateLocalInputCmd.Broadcast();
 		PendingInputCmd.bLegit = true;
-		//PendingInputCmd.Counter++;
+		PendingInputCmd.Counter++;
 	}
 }
 
diff --git a/Engine/Plugins/Runtime/NetworkPredictionExtras/Source/NetworkPredictionExtras/Public/PhysicsMovement.h b/Engine/Plugins/Runtime/NetworkPredictionExtras/Source/NetworkPredictionExtras/Public/PhysicsMovement.h
--- a/Engine/Plugins/Runtime/NetworkPredictionExtras/Source/NetworkPredictionExtras/Public/PhysicsMovement.h
+++ b/Engine/Plugins/Runtime/NetworkPredictionExtras/Source/NetworkPredictionExtras/Public/PhysicsMovement.h
@@ -11,7 +11,7 @@
 
 namespace UE_NP
 {
-	NETSIM_DEVCVAR_SHIPCONST_INT(ForceReconcilePhysicsInputCmd, 0, "np2.PhysicsInputCmdForceReconcile", "");
+	
 	NETSIM_DEVCVAR_SHIPCONST_INT(ForceReconcilePhysicsMovementState, 0, "np2.ForceReconcilePhysicsMovementState", "");
 };
 
@@ -78,23 +78,11 @@ struct FPhysicsInputCmd
 		}
 	}
 
-	bool ShouldReconcile(const FPhysicsInputCmd& AuthState) const
-	{
-		return FVector::DistSquared(Force, AuthState.Force) > 0.1f
-			|| FVector::DistSquared(Torque, AuthState.Torque) > 0.1f
-			|| FVector::DistSquared(Acceleration, AuthState.Acceleration) > 0.1f
-			|| FVector::DistSquared(AngularAcceleration, AuthState.AngularAcceleration) > 0.1f
-			|| !FMath::IsNearlyEqual(TargetYaw, AuthState.TargetYaw, 1.0f)
-			|| bJumpedPressed != AuthState.bJumpedPressed
-			//|| Counter != AuthState.Counter // this will cause constant corrections with multiple clients but useful in testing single client
-			|| bBrakesPressed != AuthState.bBrakesPressed
-			|| (UE_NP::ForceReconcilePhysicsInputCmd() > 0);
-			
-	}
+	bool ShouldReconcile(const FPhysicsInputCmd& AuthState) const;
 
 	FString ToString() const
 	{
-		return FString(TEXT(""));
+		return FString::Printf(TEXT("{Count: %d}"), Counter);
 	}
 };
 
