[Safari] Move thrust suspension impl onto physics thread.
- Update floatingpawn2 to use physics thread suspension.
- Add utility functions for applying force at position on particle.
- Fix bug in keep upright force when angle is obtuse.

diff --git a/Engine/Plugins/Runtime/NetworkPrediction/Source/NetworkPrediction/Public/NetworkPredictionCVars.h b/Engine/Plugins/Runtime/NetworkPrediction/Source/NetworkPrediction/Public/NetworkPredictionCVars.h
--- a/Engine/Plugins/Runtime/NetworkPrediction/Source/NetworkPrediction/Public/NetworkPredictionCVars.h
+++ b/Engine/Plugins/Runtime/NetworkPrediction/Source/NetworkPrediction/Public/NetworkPredictionCVars.h
@@ -115,7 +115,7 @@ namespace UE_NETWORK_PHYSICS
 	NETSIM_DEVCVAR_SHIPCONST_FLOAT(MaxAngularVelocity, 30.f, "np2.Mock.MaxAngularVelocity", "Limits how fast character can possibly rotate.")
 	NETSIM_DEVCVAR_SHIPCONST_FLOAT(MovementK, 1.25, "np2.Mock.MovementK", "Movement Coefficient (higher=faster movement)");
 	NETSIM_DEVCVAR_SHIPCONST_FLOAT(RotationK, 1.25, "np2.Mock.RotationK", "Rotation Coefficient (higher=faster movement)");
-	NETSIM_DEVCVAR_SHIPCONST_FLOAT(TurnK, 10000000.f, "np2.Mock.TurnK", "Coefficient for automatic turning (higher=quicker turning)");
+	NETSIM_DEVCVAR_SHIPCONST_FLOAT(TurnK, 100000000.f, "np2.Mock.TurnK", "Coefficient for automatic turning (higher=quicker turning)");
 	NETSIM_DEVCVAR_SHIPCONST_FLOAT(TurnDampK, 100.f, "np2.Mock.TurnDampK", "Coefficient for damping portion of turn. Higher=more damping but too higher will lead to instability.");
 	NETSIM_DEVCVAR_SHIPCONST_FLOAT(DampYawVelocityK, 2.f, "np2.Mock.DampYawVelocityK", "Coefficient for damping angular velocity in yaw direction only. This is only enabled when auto target yaw is disabled.");
 
diff --git a/Engine/Plugins/Runtime/NetworkPredictionExtras/Source/NetworkPredictionExtras/Private/PhysicsMovement.cpp b/Engine/Plugins/Runtime/NetworkPredictionExtras/Source/NetworkPredictionExtras/Private/PhysicsMovement.cpp
--- a/Engine/Plugins/Runtime/NetworkPredictionExtras/Source/NetworkPredictionExtras/Private/PhysicsMovement.cpp
+++ b/Engine/Plugins/Runtime/NetworkPredictionExtras/Source/NetworkPredictionExtras/Private/PhysicsMovement.cpp
@@ -42,6 +42,10 @@ NETSIM_DEVCVAR_SHIPCONST_INT(DisableAngularVelLimit, 0, "np2.PhysicsMovement.Dis
 NETSIM_DEVCVAR_SHIPCONST_INT(ForceReconcilePhysicsInputCmd, 0, "np2.PhysicsInputCmdForceReconcile", "");
 NETSIM_DEVCVAR_SHIPCONST_FLOAT(TargetYawTolerance, 0.1f, "np2.PhysicsMovement.TargetYawTolerance", "");
 
+NETSIM_DEVCVAR_SHIPCONST_INT(DisableSprings, 0, "np2.PhysicsMovement.DisableSprings", "");
+NETSIM_DEVCVAR_SHIPCONST_INT(SpringTraceWorldDown, 1, "np2.PhysicsMovement.Spring.TraceWorldDown", "If 1: trace straight down, if 0, trace in direction of springs (local down) instead of world down.");
+NETSIM_DEVCVAR_SHIPCONST_INT(SpringPushOnGround, 0, "np2.PhysicsMovement.Spring.PushOnGround", "If enabled, springs apply equal force against object below character.");
+
 bool FPhysicsInputCmd::ShouldReconcile(const FPhysicsInputCmd& AuthState) const
 {
 	return FVector::DistSquared(Force, AuthState.Force) > 0.1f
@@ -100,24 +104,20 @@ struct FPhysicsMovementSimulation
 			UE_LOG(LogTemp, Warning, TEXT("[%d][%d] Location: %s"), UE_NETWORK_PHYSICS::DebugSimulationFrame(), UE_NETWORK_PHYSICS::DebugServer(), *PT->X().ToString());
 		}
 
-		FVector TracePosition = PT->X();
-		FVector EndPosition = TracePosition + FVector(0.f, 0.f, -100.f);
-		FCollisionShape Shape = FCollisionShape::MakeSphere(250.f);
-		ECollisionChannel CollisionChannel = ECollisionChannel::ECC_WorldStatic;
-		FCollisionResponseParams ResponseParams = FCollisionResponseParams::DefaultResponseParam;
-		FCollisionObjectQueryParams ObjectParams(ECollisionChannel::ECC_PhysicsBody);
-
-		FHitResult OutHit;
-		bool bInAir = !UE_NETWORK_PHYSICS::JumpHack() && !World->LineTraceSingleByChannel(OutHit, TracePosition, EndPosition, ECollisionChannel::ECC_WorldStatic, LocalState.QueryParams, ResponseParams);
-		const float UpDot = FVector::DotProduct(PT->R().GetUpVector(), FVector::UpVector);
-
-		if (UPrimitiveComponent* HitPrim = OutHit.Component.Get())
+		bool bInAir = false;
 		{
-			if (FPhysicsInterface::IsValid(HitPrim->BodyInstance.ActorHandle) && (LocalState.Proxy == HitPrim->BodyInstance.ActorHandle))
-			{
-				ensure(false); // hit self somehow?
-			}
+			FHitResult OutHit;
+			FVector TracePosition = PT->X();
+			FVector EndPosition = TracePosition + FVector(0.f, 0.f, -100.f);
+			FCollisionShape Shape = FCollisionShape::MakeSphere(250.f);
+			ECollisionChannel CollisionChannel = ECollisionChannel::ECC_WorldStatic;
+			FCollisionResponseParams ResponseParams = FCollisionResponseParams::DefaultResponseParam;
+			FCollisionObjectQueryParams ObjectParams(ECollisionChannel::ECC_PhysicsBody);
+
+			bInAir = !UE_NETWORK_PHYSICS::JumpHack() && !World->LineTraceSingleByChannel(OutHit, TracePosition, EndPosition, ECollisionChannel::ECC_WorldStatic, LocalState.QueryParams, ResponseParams);
 		}
+		
+		const float UpDot = FVector::DotProduct(PT->R().GetUpVector(), FVector::UpVector);
 
 		// Debug CVar to make jump cause mispredictions
 		if (UE_NETWORK_PHYSICS::JumpMisPredict())
@@ -200,15 +200,31 @@ struct FPhysicsMovementSimulation
 		}
 		
 		// Keep pawn upright
-		if (DisableKeepUpright() == 0 && NetState.bEnableKeepUpright && UpDot < 0.95f)
+		if (DisableKeepUpright() == 0 && NetState.bEnableKeepUpright && UpDot < NetState.KeepUprightThreshold)
 		{
 			FRotator Rot = PT->R().Rotator();
 			FVector Target(0, 0, 1); // TODO choose up based on floor normal
 			FVector CurrentUp = Rot.RotateVector(Target);
-			FVector Torque = -FVector::CrossProduct(Target, CurrentUp);
+
+			// Get axis of rotation
+			FVector TorqueDirection = -FVector::CrossProduct(Target, CurrentUp);
+			float Length2 = TorqueDirection.SizeSquared();
+			if (!FMath::IsNearlyZero(Length2))
+			{
+				TorqueDirection /= FMath::Sqrt(Length2);
+			}
+			else
+			{
+				TorqueDirection = FVector(1, 0, 0);
+			}
+
+			// Get cos(angle) w/ dot and map cos([0,180]) from [1,-1] to [0,1], we want max rotation in opposite directions, no rotation in same direction.
+			float Rotation = (FVector::DotProduct(Target, CurrentUp) - 1) * -0.5f;
+
+			FVector Torque = TorqueDirection * Rotation * UE_NETWORK_PHYSICS::TurnK();
 			FVector Damp(PT->W().X* UE_NETWORK_PHYSICS::TurnDampK(), PT->W().Y* UE_NETWORK_PHYSICS::TurnDampK(), 0);
 
-			PT->AddTorque(Torque* UE_NETWORK_PHYSICS::TurnK() + Damp);
+			PT->AddTorque(Torque + Damp);
 		}
 
 		if (InputCmd.bBrakesPressed)
@@ -299,6 +315,61 @@ struct FPhysicsMovementSimulation
 			}
 		}
 
+		// Evaluate Springs
+		if (DisableSprings() == 0)
+		{
+			for (FSpring& Spring : NetState.Springs)
+			{	
+				if (Spring.RaycastLength <= 0.0f)
+				{
+					continue;
+				}
+				const Chaos::FRigidTransform3 CurrentTransform = Chaos::FParticleUtilitiesXR::GetActorWorldTransform(PT);
+				const Chaos::FVec3 WorldSpringPosition = CurrentTransform.TransformPosition(Spring.LocalPosition);
+
+				const FVector& StartPosition = WorldSpringPosition;
+				const Chaos::FVec3 WorldSpringDirection = CurrentTransform.TransformVector(Spring.LocalDirection);
+				const FVector TraceDirection = SpringTraceWorldDown() ? Chaos::FVec3(0, 0, -1) : -WorldSpringDirection;
+				const FVector EndPosition = StartPosition + TraceDirection * Spring.RaycastLength;
+
+				FHitResult OutHit;
+				if (World->LineTraceSingleByChannel(OutHit, StartPosition, EndPosition, Spring.TraceChannel, LocalState.QueryParams))
+				{	
+					const float CurrentLength = OutHit.Time * Spring.RaycastLength;
+					const float DeltaLength = CurrentLength - Spring.NaturalLength;
+					const float DeltaLengthDerivative = (CurrentLength - Spring.GetPreviousLength()) / DeltaSeconds;
+					
+					Chaos::FVec3 WorldForce = WorldSpringDirection * (DeltaLength * Spring.Stiffness + DeltaLengthDerivative * Spring.DampingStrength);
+					if (NetState.bScaleSpringForcesWithMass)
+					{
+						WorldForce *= PT->M();
+					}
+
+					// Apply force at relative location on character
+					{
+
+						Chaos::FParticleUtilitiesXR::AddForceAtPositionWorld(PT, WorldForce, WorldSpringPosition);
+					}
+
+					// Apply force on object spring is pushing on.
+					if (SpringPushOnGround())
+					{
+						// Note: Not threadsafe until we get PT SQ API. Should be safe unless object is destroyed.
+						Chaos::FRigidBodyHandle_Internal* HitObject = OutHit.GetComponent()->BodyInstance.ActorHandle->GetPhysicsThreadAPI();
+						Chaos::FParticleUtilitiesXR::AddForceAtPositionWorld(HitObject, -WorldForce, OutHit.Location);
+					}
+
+					// Cache length for next frame.
+					Spring.SetPreviousLength(CurrentLength);
+				}
+				else
+				{
+					// Cache length for next frame.
+					Spring.SetPreviousLength(Spring.NaturalLength);
+				}
+			}
+		}
+
 		NetState.JumpCooldownMS = FMath::Max( NetState.JumpCooldownMS - (int32)(DeltaSeconds* 1000.f), 0);
 		if (NetState.JumpCooldownMS != 0)
 		{
@@ -450,6 +521,14 @@ void UPhysicsMovementComponent::SetEnableKeepUpright(bool bKeepUpright)
 	});
 }
 
+void UPhysicsMovementComponent::SetKeepUprightThreshold(float Threshold)
+{
+	NetworkPredictionProxy.ModifyNetState<FPhysicsMovementAsyncModelDef>([Threshold](FPhysicsMovementNetState& NetState)
+	{
+			NetState.KeepUprightThreshold = FMath::Clamp(Threshold, 0.0f, 1.0f);
+	});
+}
+
 void UPhysicsMovementComponent::SetAutoBrakeStrength(float AutoBrakeStrength)
 {
 	NetworkPredictionProxy.ModifyNetState<FPhysicsMovementAsyncModelDef>([AutoBrakeStrength](FPhysicsMovementNetState& NetState)
@@ -458,6 +537,31 @@ void UPhysicsMovementComponent::SetAutoBrakeStrength(float AutoBrakeStrength)
 	});
 }
 
+void UPhysicsMovementComponent::SetScaleSpringForcesWithMass(bool bScaleWithMass)
+{
+	NetworkPredictionProxy.ModifyNetState<FPhysicsMovementAsyncModelDef>([bScaleWithMass](FPhysicsMovementNetState& NetState)
+	{
+			NetState.bScaleSpringForcesWithMass = bScaleWithMass;
+	});
+}
+
+void UPhysicsMovementComponent::AddSpring(FSpring Spring)
+{
+	Spring.ResetState(); // Make sure previous length and other computed data is correct for first frame.
+	NetworkPredictionProxy.ModifyNetState<FPhysicsMovementAsyncModelDef>([Spring](FPhysicsMovementNetState& NetState) mutable
+	{
+		NetState.Springs.Add(Spring);
+	});
+}
+
+void UPhysicsMovementComponent::ClearSprings()
+{
+	NetworkPredictionProxy.ModifyNetState<FPhysicsMovementAsyncModelDef>([](FPhysicsMovementNetState& NetState)
+	{
+			NetState.Springs.Reset();
+	});
+}
+
 
 void UPhysicsMovementComponent::TestMisprediction()
 {
diff --git a/Engine/Plugins/Runtime/NetworkPredictionExtras/Source/NetworkPredictionExtras/Public/PhysicsMovement.h b/Engine/Plugins/Runtime/NetworkPredictionExtras/Source/NetworkPredictionExtras/Public/PhysicsMovement.h
--- a/Engine/Plugins/Runtime/NetworkPredictionExtras/Source/NetworkPredictionExtras/Public/PhysicsMovement.h
+++ b/Engine/Plugins/Runtime/NetworkPredictionExtras/Source/NetworkPredictionExtras/Public/PhysicsMovement.h
@@ -7,6 +7,7 @@
 #include "Async/NetworkPredictionAsyncProxy.h"
 #include "NetworkPhysicsComponent.h"
 #include "NetworkPredictionCVars.h"
+#include "Springs.h"
 #include "PhysicsMovement.generated.h"
 
 namespace UE_NP
@@ -112,10 +113,19 @@ struct FPhysicsMovementNetState
 	UPROPERTY(BlueprintReadWrite, Category = "Mock Object")
 	bool bEnableKeepUpright = true;
 
+	// Value 0 to 1 indicating when forces are applied to keep character upright. 0 does nothing, 1 always tries to align character with up,
+	// 0.5 will start applying forces when tipped 90 degrees, etc.
+	UPROPERTY(BlueprintReadWrite, Category = "Mock Object")
+	float KeepUprightThreshold = 0.95f;
+
 	// Strength of auto brake force applied when no input force and on ground.
 	UPROPERTY(BlueprintReadWrite, Category = "Mock Object")
 	float AutoBrakeStrength = 5.f;
 
+	// Should spring forces scale with mass?
+	UPROPERTY(BlueprintReadWrite, Category = "Mock Object")
+	bool bScaleSpringForcesWithMass = true;
+
 	// Arbitrary data that doesn't affect sim but could still trigger rollback
 	UPROPERTY(BlueprintReadWrite, Category="Mock Object")
 	int32 RandValue = 0;
@@ -146,6 +156,9 @@ struct FPhysicsMovementNetState
 	UPROPERTY(BlueprintReadWrite, Category="Mock Object")
 	int32 KickFrame = 0;
 
+	UPROPERTY(BlueprintReadWrite, Category="Mock Object")
+	TArray<FSpring> Springs;
+
 	void NetSerialize(FArchive& Ar)
 	{
 		Ar << ForceMultiplier;
@@ -155,7 +168,9 @@ struct FPhysicsMovementNetState
 		Ar << AutoFaceTargetYawDamp;
 		Ar << bEnableAutoFaceTargetYaw;
 		Ar << AutoBrakeStrength;
+		Ar << bScaleSpringForcesWithMass;
 		Ar << bEnableKeepUpright;
+		Ar << KeepUprightThreshold;
 		Ar << JumpCooldownMS;
 		Ar << JumpCount;
 		Ar << CheckSum;	
@@ -163,6 +178,7 @@ struct FPhysicsMovementNetState
 		Ar << JumpStartFrame;
 		Ar << InAirFrame;
 		Ar << KickFrame;
+		Ar << Springs;
 	}
 
 	bool ShouldReconcile(const FPhysicsMovementNetState& AuthState) const
@@ -175,7 +191,9 @@ struct FPhysicsMovementNetState
 			AutoFaceTargetYawDamp != AuthState.AutoFaceTargetYawDamp ||
 			bEnableAutoFaceTargetYaw != AuthState.bEnableAutoFaceTargetYaw ||
 			AutoBrakeStrength != AuthState.AutoBrakeStrength ||
+			bScaleSpringForcesWithMass != AuthState.bScaleSpringForcesWithMass || 
 			bEnableKeepUpright != AuthState.bEnableKeepUpright ||
+			KeepUprightThreshold != AuthState.KeepUprightThreshold || 
 			JumpCooldownMS != AuthState.JumpCooldownMS || 
 			JumpCount != AuthState.JumpCount ||
 			RecoveryFrame != AuthState.RecoveryFrame ||
@@ -183,6 +201,7 @@ struct FPhysicsMovementNetState
 			InAirFrame != AuthState.InAirFrame ||
 			KickFrame!= AuthState.KickFrame ||
 			CheckSum != AuthState.CheckSum ||
+			Springs != AuthState.Springs || 
 			(UE_NP::ForceReconcilePhysicsMovementState() > 0);
 	}
 };
@@ -242,9 +261,22 @@ public:
 	UFUNCTION(BlueprintCallable, Category="Movement")
 	void SetEnableKeepUpright(bool bKeepUpright);
 
+	UFUNCTION(BlueprintCallable, Category="Movement")
+	void SetKeepUprightThreshold(float Threshold);
+
 	UFUNCTION(BlueprintCallable, Category="Movement")
 	void SetAutoBrakeStrength(float BrakeStrength);
 
+	UFUNCTION(BlueprintCallable, Category = "Movement")
+	void SetScaleSpringForcesWithMass(bool bScaleWithMass);
+	
+	// Add a spring, spring is persistant, can be removed by ClearSprings().
+	UFUNCTION(BlueprintCallable, Category="Movement")
+	void AddSpring(FSpring Spring);
+
+	UFUNCTION(BlueprintCallable, Category = "Movement")
+	void ClearSprings();
+
 	void TestMisprediction();
 
 private:
diff --git a/Engine/Plugins/Runtime/NetworkPredictionExtras/Source/NetworkPredictionExtras/Public/Springs.h b/Engine/Plugins/Runtime/NetworkPredictionExtras/Source/NetworkPredictionExtras/Public/Springs.h
new file mode 100644
--- /dev/null
+++ b/Engine/Plugins/Runtime/NetworkPredictionExtras/Source/NetworkPredictionExtras/Public/Springs.h
@@ -0,0 +1,87 @@
+// Copyright Epic Games, Inc. All Rights Reserved.
+
+#pragma once
+#include "Springs.generated.h"
+
+
+struct FSpring;
+static FArchive& operator<<(FArchive& Ar, FSpring& Spring);
+
+USTRUCT(BlueprintType)
+struct FSpring
+{
+	GENERATED_BODY()
+
+	friend FArchive& operator<<(FArchive& Ar, FSpring& Spring);
+
+	FSpring()
+		: LocalPosition(ForceInitToZero)
+		, LocalDirection(ForceInitToZero)
+		, TraceChannel(ECC_Visibility)
+		, Stiffness(0.0f)
+		, DampingStrength(0.0f)
+		, RaycastLength(0.0f)
+		, NaturalLength(0.0f)
+		, PreviousLength(0.0f)
+	{}
+
+	bool operator==(const FSpring& Spring) const
+	{
+		 return LocalPosition == Spring.LocalPosition
+			&& LocalDirection == Spring.LocalDirection
+			&& TraceChannel == Spring.TraceChannel
+			&& Stiffness == Spring.Stiffness
+			&& DampingStrength == Spring.DampingStrength
+			&& RaycastLength == Spring.RaycastLength
+			&& NaturalLength == Spring.NaturalLength;
+			//&& PreviousLength == Spring.PreviousLength; // Deliberately not reconciled, not replicating this value.
+	}
+
+	// Position of spring relative to actor, raycast starts here and force is applied at this position.
+	UPROPERTY(BlueprintReadWrite, Category="Spring")
+	FVector LocalPosition;
+
+	// Direction of spring force applied to body.
+	UPROPERTY(BlueprintReadWrite, Category = "Spring")
+	FVector LocalDirection;
+
+	// Trace channel of raycast
+	UPROPERTY(BlueprintReadWrite, Category = "Spring")
+	TEnumAsByte<ECollisionChannel> TraceChannel;
+	
+	// Strength of spring force.
+	UPROPERTY(BlueprintReadWrite, Category = "Spring")
+	float Stiffness;
+
+	// Strength of damping term.
+	UPROPERTY(BlueprintReadWrite, Category = "Spring")
+	float DampingStrength;
+
+	// Lenght of raycast.
+	UPROPERTY(BlueprintReadWrite, Category = "Spring")
+	float RaycastLength;
+
+	// Natural length of spring when it is not stretched/compressed.
+	UPROPERTY(BlueprintReadWrite, Category = "Spring")
+	float NaturalLength;
+
+	float GetPreviousLength() const { return PreviousLength; }
+	void SetPreviousLength(float InLength)  { PreviousLength = InLength; }
+	void ResetState() { SetPreviousLength(NaturalLength); }
+private:
+	// Length of spring from last step.
+	float PreviousLength;
+};
+
+static FArchive& operator<<(FArchive& Ar, FSpring& Spring)
+{
+	Ar << Spring.LocalPosition;
+	Ar << Spring.LocalDirection;
+	Ar << Spring.TraceChannel;
+	Ar << Spring.Stiffness;
+	Ar << Spring.DampingStrength;
+	Ar << Spring.RaycastLength;
+	Ar << Spring.NaturalLength;
+	//Ar << Spring.PreviousLength; Deliberately not serialized as we do not need to replicate.
+	return Ar;
+}
diff --git a/Engine/Source/Runtime/Experimental/Chaos/Public/Chaos/Particle/ParticleUtilities.h b/Engine/Source/Runtime/Experimental/Chaos/Public/Chaos/Particle/ParticleUtilities.h
--- a/Engine/Source/Runtime/Experimental/Chaos/Public/Chaos/Particle/ParticleUtilities.h
+++ b/Engine/Source/Runtime/Experimental/Chaos/Public/Chaos/Particle/ParticleUtilities.h
@@ -176,6 +176,26 @@ namespace Chaos
 			TSpatialAccessor::Position(Particles, Index) = P;
 			TSpatialAccessor::Rotation(Particles, Index) = Q;
 		}
+		
+		template<typename T_PARTICLEHANDLE>
+		static inline void AddForceAtPositionLocal(T_PARTICLEHANDLE Particle, const FVec3& LocalForce, const FVec3& LocalPosition)
+		{
+			const FRigidTransform3 ParticleTransform = GetActorWorldTransform(Particle);
+			const FVec3 WorldPosition = ParticleTransform.TransformPosition(LocalPosition);
+			const FVec3 WorldForce = ParticleTransform.TransformVector(LocalForce);
+
+			AddForceAtPositionWorld(Particle, WorldForce, WorldPosition);
+		}
+
+		template<typename T_PARTICLEHANDLE>
+		static inline void AddForceAtPositionWorld(T_PARTICLEHANDLE Particle, const FVec3& Force, const FVec3& Position)
+		{
+			const FVec3 WorldCOM = GetCoMWorldPosition(Particle);
+			const FVec3 Torque = FVec3::CrossProduct(Position - WorldCOM, Force);
+
+			Particle->AddForce(Force);
+			Particle->AddTorque(Torque);
+		}
 	};
 
 	/**
