Fix non-unity CIS error in NetworkPredictionExtras/Springs.h

diff --git a/Engine/Plugins/Runtime/NetworkPredictionExtras/Source/NetworkPredictionExtras/Public/Springs.h b/Engine/Plugins/Runtime/NetworkPredictionExtras/Source/NetworkPredictionExtras/Public/Springs.h
--- a/Engine/Plugins/Runtime/NetworkPredictionExtras/Source/NetworkPredictionExtras/Public/Springs.h
+++ b/Engine/Plugins/Runtime/NetworkPredictionExtras/Source/NetworkPredictionExtras/Public/Springs.h
@@ -1,6 +1,7 @@
 // Copyright Epic Games, Inc. All Rights Reserved.
 
 #pragma once
+#include "Engine/EngineTypes.h"
 #include "Springs.generated.h"
 
 
