Network Physics:
"Physics Effects". Basic API for doing gameplay calculations on physics thread.

diff --git a/Engine/Plugins/Runtime/NetworkPrediction/Source/NetworkPrediction/Private/Async/NetworkPredictionAsyncProxy.cpp b/Engine/Plugins/Runtime/NetworkPrediction/Source/NetworkPrediction/Private/Async/NetworkPredictionAsyncProxy.cpp
--- a/Engine/Plugins/Runtime/NetworkPrediction/Source/NetworkPrediction/Private/Async/NetworkPredictionAsyncProxy.cpp
+++ b/Engine/Plugins/Runtime/NetworkPrediction/Source/NetworkPrediction/Private/Async/NetworkPredictionAsyncProxy.cpp
@@ -81,6 +81,24 @@ void FNetworkPredictionAsyncProxy::OnPreReplication()
 {
 	if (Manager)
 	{
-		CachedLatestFrame = Manager->GetLatestFrame();
+		CachedLatestFrame = Manager->GetLatestSimFrame();
 	}
+}
+
+int32 FNetworkPredictionAsyncProxy::GetNextSimFrame() const
+{
+	if (Manager)
+	{
+		return Manager->GetNextSimFrame();
+	}
+	return INDEX_NONE;
+}
+	
+int32 FNetworkPredictionAsyncProxy::GetLatestOutputSimFrame() const
+{
+	if (Manager)
+	{
+		return Manager->GetLatestSimFrame();
+	}
+	return INDEX_NONE;
 }
\ No newline at end of file
diff --git a/Engine/Plugins/Runtime/NetworkPrediction/Source/NetworkPrediction/Public/Async/NetworkPredictionAsyncInstanceData.h b/Engine/Plugins/Runtime/NetworkPrediction/Source/NetworkPrediction/Public/Async/NetworkPredictionAsyncInstanceData.h
--- a/Engine/Plugins/Runtime/NetworkPrediction/Source/NetworkPrediction/Public/Async/NetworkPredictionAsyncInstanceData.h
+++ b/Engine/Plugins/Runtime/NetworkPrediction/Source/NetworkPrediction/Public/Async/NetworkPredictionAsyncInstanceData.h
@@ -34,6 +34,7 @@ struct TAsncInstanceStaticData
 
 	uint16 Index;
 	LocalStateType LocalState;
+	int32 LocalSpawnFrame = INDEX_NONE;
 };
 
 template<typename AsyncModelDef>
diff --git a/Engine/Plugins/Runtime/NetworkPrediction/Source/NetworkPrediction/Public/Async/NetworkPredictionAsyncModelDef.h b/Engine/Plugins/Runtime/NetworkPrediction/Source/NetworkPrediction/Public/Async/NetworkPredictionAsyncModelDef.h
--- a/Engine/Plugins/Runtime/NetworkPrediction/Source/NetworkPrediction/Public/Async/NetworkPredictionAsyncModelDef.h
+++ b/Engine/Plugins/Runtime/NetworkPrediction/Source/NetworkPrediction/Public/Async/NetworkPredictionAsyncModelDef.h
@@ -1,6 +1,8 @@
 // Copyright Epic Games, Inc. All Rights Reserved.
 
 #pragma once
+#include "Containers/UnrealString.h"
+#include "Templates/EnableIf.h"
 
 class AActor;
 class APlayerController;
@@ -35,4 +37,43 @@ struct FNetworkPredictionAsyncModelDef
 	using ControlKeyType = typename AsyncModelDef::ControlKeyType; \
 	using SimulationTickType = typename AsyncModelDef::SimulationTickType;
 
+
+
+// Template nonsense to implement optional user-defined functions and operations.
+// Mainly we want to avoid base classes on the user states, so calls that are optional
+// can go through NpModelUtil to either do nothing or some default implementation, or
+// be overridden by the user when they need to.
+template<typename AsyncModelDef>
+struct NpModelUtilBase
+{
+	HOIST_ASYNCMODELDEF_TYPES();
+
+	
+	struct CToStringFuncable
+	{
+		template <typename InType>
+		auto Requires(InType* T) -> decltype(T->ToString());
+	};
+
+	template<typename T, bool bEnable=TModels<CToStringFuncable, T>::Value>
+	static typename TEnableIf<bEnable, FString>::Type ToString(const T* State)
+	{
+		return State->ToString();
+	}
+
+	
+	template<typename T, bool bEnable=TModels<CToStringFuncable, T>::Value>
+	static typename TEnableIf<!bEnable, FString>::Type ToString(const T* State)
+	{
+		return FString(TEXT("{requires ToString()}"));
+	}
+};
+
+template<typename AsyncModelDef>
+struct NpModelUtil : NpModelUtilBase<AsyncModelDef>
+{
+
+};
+
+
 } // namespace UE_NP
\ No newline at end of file
diff --git a/Engine/Plugins/Runtime/NetworkPrediction/Source/NetworkPrediction/Public/Async/NetworkPredictionAsyncProxy.h b/Engine/Plugins/Runtime/NetworkPrediction/Source/NetworkPrediction/Public/Async/NetworkPredictionAsyncProxy.h
--- a/Engine/Plugins/Runtime/NetworkPrediction/Source/NetworkPrediction/Public/Async/NetworkPredictionAsyncProxy.h
+++ b/Engine/Plugins/Runtime/NetworkPrediction/Source/NetworkPrediction/Public/Async/NetworkPredictionAsyncProxy.h
@@ -30,6 +30,13 @@ struct NETWORKPREDICTION_API FNetworkPredictionAsyncProxy
 	bool Identical(const FNetworkPredictionAsyncProxy* Other, uint32 PortFlags) const;
 	void OnPreReplication();
 
+	// Gets the next expected simulation frame number that will be marshalled GT -> PT
+	// Note this can be unstable on clients. For example if you get this value and then receive an input fault from server, the next frame may be different.
+	int32 GetNextSimFrame() const;
+
+	// Returns the latest simulation frame that was marshalled back from the PT
+	int32 GetLatestOutputSimFrame() const;
+
 	// Include NetworkPredictionAsyncProxyImpl.h in your .cpp to use these.
 	template<typename AsyncModelDef>
 	void RegisterSim(typename AsyncModelDef::LocalStateType&& LocalState, typename AsyncModelDef::NetStateType&& NetState, typename AsyncModelDef::InputCmdType* PendingInputCmd, typename AsyncModelDef::NetStateType* OutNetState);
diff --git a/Engine/Plugins/Runtime/NetworkPrediction/Source/NetworkPrediction/Public/Async/NetworkPredictionAsyncWorldManager.h b/Engine/Plugins/Runtime/NetworkPrediction/Source/NetworkPrediction/Public/Async/NetworkPredictionAsyncWorldManager.h
--- a/Engine/Plugins/Runtime/NetworkPrediction/Source/NetworkPrediction/Public/Async/NetworkPredictionAsyncWorldManager.h
+++ b/Engine/Plugins/Runtime/NetworkPrediction/Source/NetworkPrediction/Public/Async/NetworkPredictionAsyncWorldManager.h
@@ -426,7 +426,7 @@ public:
 		npCheckSlow(AsyncCallback);
 		FNetworkPredictionSimCallbackInput* AsyncInput = AsyncCallback->GetProducerInputData_External();
 		npCheckSlow(AsyncInput);
-
+		
 		AsyncCallback->PushLatestLocalFrameOffset(LocalFrameOffset);
 
 		for (int32 StepNum=0; StepNum < NumSteps; ++StepNum)
@@ -461,6 +461,7 @@ public:
 		}
 
 		NextPhysicsStep = PhysicsStep;
+		NextSimulationStep = PhysicsStep - LocalFrameOffset;
 
 		// Consume all output data. Note we are consuming all SimCallback Output objects here even though we really only want the latest -
 		// but there maybe objects that were destroyed and not present in the final output, so this is trying to not miss them.
@@ -489,11 +490,16 @@ public:
 	}
 
 	// Returns latest simulation frame number that has been marshaled back from PT
-	int32 GetLatestFrame() const
+	int32 GetLatestSimFrame() const
 	{
 		return GlobalOutput_External.SimulationFrame;
 	}
 
+	int32 GetNextSimFrame() const
+	{
+		return NextSimulationStep;
+	}
+
 private:
 
 	struct FGlobalInstanceData
@@ -572,7 +578,8 @@ private:
 
 	int32 SpawnIDCounter = 1;
 	int32 ClientSpawnIDCounter = -1;
-	int32 NextPhysicsStep = 0;
+	int32 NextPhysicsStep = 0; // Local
+	int32 NextSimulationStep = 0; // Networked
 
 	uint32 InitializedModelDefsMask = 0;
 
diff --git a/Engine/Plugins/Runtime/NetworkPrediction/Source/NetworkPrediction/Public/Async/Services/NetworkPredictionAsyncService_Marshal.inl b/Engine/Plugins/Runtime/NetworkPrediction/Source/NetworkPrediction/Public/Async/Services/NetworkPredictionAsyncService_Marshal.inl
--- a/Engine/Plugins/Runtime/NetworkPrediction/Source/NetworkPrediction/Public/Async/Services/NetworkPredictionAsyncService_Marshal.inl
+++ b/Engine/Plugins/Runtime/NetworkPrediction/Source/NetworkPrediction/Public/Async/Services/NetworkPredictionAsyncService_Marshal.inl
@@ -103,7 +103,9 @@ public:
 			// New instances
 			for (const typename TAsyncModelDataStore_Input<AsyncModelDef>::FNewInstance& NewInstance : InputData->NewInstances)
 			{
-				DataStore->Instances.FindOrAdd(NewInstance.ID) = NewInstance.StaticData;
+				TAsncInstanceStaticData<AsyncModelDef>& StaticData = DataStore->Instances.FindOrAdd(NewInstance.ID);
+				StaticData = NewInstance.StaticData;
+				StaticData.LocalSpawnFrame = LocalFrame;
 
 				NpResizeForIndex(Snapshot.InputCmds, NewInstance.StaticData.Index);
 				NpResizeForIndex(Snapshot.NetStates, NewInstance.StaticData.Index);
@@ -190,6 +192,8 @@ public:
 
 			NpResizeForIndex(Snapshot.InputCmds, idx);
 			NpResizeForIndex(Snapshot.NetStates, idx);
+
+			npEnsure(Snapshot.InputCmds.Num()==Snapshot.NetStates.Num());
 				
 			Snapshot.InputCmds[idx] = CorrectionInstance.InputCmd;
 			Snapshot.NetStates[idx] = CorrectionInstance.NetState;
diff --git a/Engine/Plugins/Runtime/NetworkPrediction/Source/NetworkPrediction/Public/Async/Services/NetworkPredictionAsyncService_NetSerialize.inl b/Engine/Plugins/Runtime/NetworkPrediction/Source/NetworkPrediction/Public/Async/Services/NetworkPredictionAsyncService_NetSerialize.inl
--- a/Engine/Plugins/Runtime/NetworkPrediction/Source/NetworkPrediction/Public/Async/Services/NetworkPredictionAsyncService_NetSerialize.inl
+++ b/Engine/Plugins/Runtime/NetworkPrediction/Source/NetworkPrediction/Public/Async/Services/NetworkPredictionAsyncService_NetSerialize.inl
@@ -31,6 +31,13 @@ public:
 
 	void NetRecv(FNetworkPredictionAsyncID ID, uint8 Flags, int32 ThisFrame, int32 FutureFrame, FNetBitReader& Ar) final override
 	{
+		uint8 bIsValid = 0;
+		Ar << bIsValid;
+		if (!bIsValid)
+		{
+			return;
+		}
+
 		TAsncInstanceStaticData<AsyncModelDef>& InstanceData = DataStore->Instances.FindChecked(ID);
 
 		if (DataStore->PendingNetRecv.NetRecvInstances.IsValidIndex(InstanceData.Index) == false)
@@ -49,8 +56,19 @@ public:
 	void NetSend(FNetworkPredictionAsyncID ID, uint8 Flags, int32 ThisFrame, int32 PendingInputCmdFrame, FNetBitWriter& Ar) final override
 	{
 		TAsncInstanceStaticData<AsyncModelDef>& InstanceData = DataStore->Instances.FindChecked(ID);
-
 		const int32 idx = InstanceData.Index;
+
+		uint8 bIsValid = true;
+		if (!DataStore->LatestSnapshot.InputCmds.IsValidIndex(idx) || !DataStore->LatestSnapshot.NetStates.IsValidIndex(idx))
+		{
+			bIsValid = false;
+		}
+		Ar << bIsValid;
+		if (!bIsValid)
+		{
+			return;
+		}
+
 		npCheckSlow(DataStore->LatestSnapshot.InputCmds.IsValidIndex(idx));
 		npCheckSlow(DataStore->LatestSnapshot.NetStates.IsValidIndex(idx));
 
@@ -75,7 +93,6 @@ private:
 		{
 			if (auto* FoundPtr = DataStore->PendingInputCmdBuffers.Find(ID))
 			{
-				// Need latest 
 				TAsyncPendingInputCmdBuffer<AsyncModelDef>& PendingInputs = FoundPtr->Get();
 				PendingInputs.Buffer[PendingInputCmdFrame % PendingInputs.Buffer.Num()].NetSerialize(Ar);
 			}
diff --git a/Engine/Plugins/Runtime/NetworkPrediction/Source/NetworkPrediction/Public/Async/Services/NetworkPredictionAsyncService_Reconcile.inl b/Engine/Plugins/Runtime/NetworkPrediction/Source/NetworkPrediction/Public/Async/Services/NetworkPredictionAsyncService_Reconcile.inl
--- a/Engine/Plugins/Runtime/NetworkPrediction/Source/NetworkPrediction/Public/Async/Services/NetworkPredictionAsyncService_Reconcile.inl
+++ b/Engine/Plugins/Runtime/NetworkPrediction/Source/NetworkPrediction/Public/Async/Services/NetworkPredictionAsyncService_Reconcile.inl
@@ -63,11 +63,15 @@ public:
 							if (LocalInputCmd.ShouldReconcile(RecvData.InputCmd))
 							{
 								UE_CLOG(LogStateCorrections() > 0, LogNetworkPrediction, Log, TEXT("Reconcile on InputCmd. %s. Instance: %d. Frame: [%d/%d]"), AsyncModelDef::GetName(), idx, RecvData.Frame - LocalFrameOffset, RecvData.Frame);
+								UE_CLOG(LogStateCorrections() > 0, LogNetworkPrediction, Log, TEXT("    Local: %s"), *NpModelUtil<AsyncModelDef>::ToString(&LocalInputCmd));
+								UE_CLOG(LogStateCorrections() > 0, LogNetworkPrediction, Log, TEXT("   Server: %s"), *NpModelUtil<AsyncModelDef>::ToString(&RecvData.InputCmd));
 								bShouldReconcile = true;
 							}
 							if (LocalNetState.ShouldReconcile(RecvData.NetState))
 							{
 								UE_CLOG(LogStateCorrections() > 0, LogNetworkPrediction, Log, TEXT("Reconcile on NetState. %s. Instance: %d. Frame: [%d/%d]"), AsyncModelDef::GetName(), idx, RecvData.Frame - LocalFrameOffset, RecvData.Frame);
+								UE_CLOG(LogStateCorrections() > 0, LogNetworkPrediction, Log, TEXT("    Local: %s"), *NpModelUtil<AsyncModelDef>::ToString(&LocalNetState));
+								UE_CLOG(LogStateCorrections() > 0, LogNetworkPrediction, Log, TEXT("   Server: %s"), *NpModelUtil<AsyncModelDef>::ToString(&RecvData.NetState));
 								bShouldReconcile = true;
 							}
 						}
diff --git a/Engine/Plugins/Runtime/NetworkPrediction/Source/NetworkPrediction/Public/Async/Services/NetworkPredictionAsyncService_Registration.inl b/Engine/Plugins/Runtime/NetworkPrediction/Source/NetworkPrediction/Public/Async/Services/NetworkPredictionAsyncService_Registration.inl
--- a/Engine/Plugins/Runtime/NetworkPrediction/Source/NetworkPrediction/Public/Async/Services/NetworkPredictionAsyncService_Registration.inl
+++ b/Engine/Plugins/Runtime/NetworkPrediction/Source/NetworkPrediction/Public/Async/Services/NetworkPredictionAsyncService_Registration.inl
@@ -145,10 +145,13 @@ private:
 		}
 
 		// Cheat a little and put the initial state values into the latest snapshot so we don't have to special case this in other spots
+		// FIXME: This isn't valid in all cases. We still have gaps!
 		NpResizeForIndex(DataStore->LatestSnapshot.NetStates, FreeIndice);
 		DataStore->LatestSnapshot.NetStates[FreeIndice] = InitialValueNet;
 		NpResizeForIndex(DataStore->LatestSnapshot.InputCmds, FreeIndice);
 
+		npEnsure(DataStore->LatestSnapshot.InputCmds.Num() == DataStore->LatestSnapshot.NetStates.Num());
+
 		// Marshal via Input
 		TAsyncModelDataStore_Input<AsyncModelDef>* InputDataStore = Input->GetPendingDataStore_External<AsyncModelDef>();
 		npCheckSlow(InputDataStore);
diff --git a/Engine/Plugins/Runtime/NetworkPrediction/Source/NetworkPrediction/Public/Async/Services/NetworkPredictionAsyncService_Ticking.inl b/Engine/Plugins/Runtime/NetworkPrediction/Source/NetworkPrediction/Public/Async/Services/NetworkPredictionAsyncService_Ticking.inl
--- a/Engine/Plugins/Runtime/NetworkPrediction/Source/NetworkPrediction/Public/Async/Services/NetworkPredictionAsyncService_Ticking.inl
+++ b/Engine/Plugins/Runtime/NetworkPrediction/Source/NetworkPrediction/Public/Async/Services/NetworkPredictionAsyncService_Ticking.inl
@@ -35,10 +35,25 @@ public:
 		{
 			FNetworkPredictionAsyncID& ID = MapIt.Key;
 			TAsncInstanceStaticData<AsyncModelDef>& InstanceData = MapIt.Value;
-			InputCmdType& InputCmd = Snapshot.InputCmds[InstanceData.Index];
-			NetStateType& NetState = Snapshot.NetStates[InstanceData.Index];
+			if (InstanceData.LocalSpawnFrame <= Context.LocalStorageFrame)			
+			{
+				if (Snapshot.InputCmds.Num() != Snapshot.NetStates.Num())
+				{
+					UE_LOG(LogNetworkPrediction, Error, TEXT("%s InputCmds and NetStates array size are out of sync (%d %d). LocalFrame: %d"), AsyncModelDef::GetName(), Snapshot.InputCmds.Num(), Snapshot.NetStates.Num(), Context.LocalStorageFrame);
+				}
 
-			SimulationTickType::Tick_Internal(Context, ID, InputCmd, InstanceData.LocalState, NetState);
+
+				if (Snapshot.InputCmds.IsValidIndex(InstanceData.Index) && Snapshot.NetStates.IsValidIndex(InstanceData.Index))
+				{
+					InputCmdType& InputCmd = Snapshot.InputCmds[InstanceData.Index];
+					NetStateType& NetState = Snapshot.NetStates[InstanceData.Index];
+					SimulationTickType::Tick_Internal(Context, ID, InputCmd, InstanceData.LocalState, NetState);
+				}
+				else
+				{
+					UE_LOG(LogNetworkPrediction, Warning, TEXT("%s Invalid Instanced %d @ index %d. LocalFrame: %d SpawnFrame: %d. Sizes: (%d %d)"), AsyncModelDef::GetName(), (int32)ID, InstanceData.Index, Context.LocalStorageFrame, InstanceData.LocalSpawnFrame, Snapshot.InputCmds.Num(), Snapshot.NetStates.Num());
+				}
+			}
 		}
 	}
 
diff --git a/Engine/Plugins/Runtime/NetworkPrediction/Source/NetworkPrediction/Public/NetworkPredictionUtil.h b/Engine/Plugins/Runtime/NetworkPrediction/Source/NetworkPrediction/Public/NetworkPredictionUtil.h
--- a/Engine/Plugins/Runtime/NetworkPrediction/Source/NetworkPrediction/Public/NetworkPredictionUtil.h
+++ b/Engine/Plugins/Runtime/NetworkPrediction/Source/NetworkPrediction/Public/NetworkPredictionUtil.h
@@ -45,6 +45,7 @@ void NpClearBitArray(BitArrayType& BitArray)
 template<typename ArrayType>
 void NpResizeForIndex(ArrayType& Array, int32 Index)
 {
+	npEnsure(Index >= 0);
 	if (Array.IsValidIndex(Index) == false)
 	{
 		Array.SetNum(Index + UE_NP::FrameStorageGrowth);
diff --git a/Engine/Plugins/Runtime/NetworkPredictionExtras/Source/NetworkPredictionExtras/Private/PhysicsEffect.cpp b/Engine/Plugins/Runtime/NetworkPredictionExtras/Source/NetworkPredictionExtras/Private/PhysicsEffect.cpp
new file mode 100644
--- /dev/null
+++ b/Engine/Plugins/Runtime/NetworkPredictionExtras/Source/NetworkPredictionExtras/Private/PhysicsEffect.cpp
@@ -0,0 +1,255 @@
+// Copyright Epic Games, Inc. All Rights Reserved.
+
+#include "PhysicsEffect.h"
+#include "Engine/World.h"
+#include "RewindData.h"
+#include "EngineUtils.h"
+#include "NetworkPredictionDebug.h"
+#include "NetworkPredictionCVars.h"
+#include "Async/NetworkPredictionAsyncWorldManager.h"
+#include "Async/NetworkPredictionAsyncID.h"
+#include "Async/NetworkPredictionAsyncModelDefRegistry.h"
+#include "Async/NetworkPredictionAsyncProxyImpl.h"
+#include "Components/PrimitiveComponent.h"
+#include "PhysicsProxy/SingleParticlePhysicsProxy.h"
+#include "DrawDebugHelpers.h"
+
+NP_DEVCVAR_INT(PhysicsEffectsDefaultFrameDelay, 5, "np2.pe.FrameDelay", "Default Frame delay for Physics Effects");
+
+struct FPhysicsEffectSimulation
+{
+	static void Tick_Internal(FNetworkPredictionAsyncTickContext& Context, UE_NP::FNetworkPredictionAsyncID ID, FPhysicsEffectCmd& InputCmd, FPhysicsEffectLocalState& LocalState, FPhysicsEffectNetState& NetState);
+};
+
+struct FPhysicsEffectsAsyncModelDef : public UE_NP::FNetworkPredictionAsyncModelDef
+{
+	NP_ASYNC_MODEL_BODY()
+
+	using InputCmdType = FPhysicsEffectCmd;
+	using NetStateType = FPhysicsEffectNetState;
+	using LocalStateType = FPhysicsEffectLocalState;
+	using SimulationTickType = FPhysicsEffectSimulation;
+
+	static const TCHAR* GetName() { return TEXT("PhysicsEffects"); }
+	static constexpr int32 GetSortPriority() { return 1; }
+};
+
+NP_ASYNC_MODEL_REGISTER(FPhysicsEffectsAsyncModelDef);
+
+void FPhysicsEffectSimulation::Tick_Internal(FNetworkPredictionAsyncTickContext& Context, UE_NP::FNetworkPredictionAsyncID ID, FPhysicsEffectCmd& InputCmd, FPhysicsEffectLocalState& LocalState, FPhysicsEffectNetState& NetState)
+{
+	UWorld* World = Context.World;
+
+	// Copy InputCmd PendingEffect if it starts after the "active" one
+	if (InputCmd.PendingEffect.StartFrame > NetState.ActiveEffect.StartFrame)
+	{
+		NetState.ActiveEffect = InputCmd.PendingEffect;
+	}
+
+	// For now, require a static (non networked) proxy in the LocalState to base everything around
+	// But we'll want to support "from X location" and not require that the "owner" of this NP system is even a physics object.
+	// For example maybe we put a PhysicsEffect sim on the GameState for world/global impulses
+	if (!LocalState.Proxy)
+	{
+		return;
+	}
+
+	auto* PT = LocalState.Proxy->GetPhysicsThreadAPI();
+	if (!PT)
+	{
+		return;
+	}
+
+	if (NetState.ActiveEffect.StartFrame != INDEX_NONE)
+	{
+		if (Context.SimulationFrame >= NetState.ActiveEffect.StartFrame)
+		{
+			if (Context.SimulationFrame <= NetState.ActiveEffect.EndFrame || NetState.ActiveEffect.EndFrame == INDEX_NONE)
+			{
+				// The effect is active
+				NetState.LastPlayedEffectTypeID = NetState.ActiveEffect.TypeID;
+				NetState.LastPlayedSimFrame = Context.SimulationFrame;
+
+				// Only sphere overlap / impulse for now
+				const float MinDistSq = NetState.ActiveEffect.QueryData.X * NetState.ActiveEffect.QueryData.X;
+
+				const FTransform ToWorldTransform(PT->R(), PT->X());
+				const FTransform QueryWorldTransform = NetState.ActiveEffect.QueryRelativeTransform * ToWorldTransform;
+				const FVector Origin = QueryWorldTransform.GetTranslation();
+
+				//DrawDebugSphere(World, Origin, 50.f, 10, FColor::Green, false, 15.f); // Not thread safe but sorta works
+
+				FCollisionShape Shape = FCollisionShape::MakeSphere(NetState.ActiveEffect.QueryData.X * 1.25f); // Extra distance + manual distance check for now
+
+				TArray<FOverlapResult> Overlaps;
+				World->OverlapMultiByChannel(Overlaps, Origin, FQuat::Identity, ECollisionChannel::ECC_PhysicsBody, Shape, LocalState.QueryParams);
+
+				//UE_LOG(LogTemp, Warning, TEXT("[%s][%d] Applying Impulse %f, %f"), World->GetNetMode() == NM_Client ? TEXT("C") : TEXT("S"), Context.SimulationFrame, NetState.ActiveEffect.QueryData.X, NetState.ActiveEffect.ResponseMagnitude.X);
+				for (FOverlapResult& Result : Overlaps)
+				{
+					//UE_LOG(LogTemp, Warning, TEXT("   Hit: %s"), *GetPathNameSafe(Result.GetActor()));
+					UPrimitiveComponent* HitPrim = Result.GetComponent();
+					
+					if (HitPrim && FPhysicsInterface::IsValid(HitPrim->BodyInstance.ActorHandle))
+					{
+						if (auto HitPT = HitPrim->BodyInstance.ActorHandle->GetPhysicsThreadAPI())
+						{
+							const FVector Delta = (HitPT->X() - Origin);
+							const float DistSq = Delta.SizeSquared();
+							if (DistSq <= MinDistSq)
+							{
+								const FVector ImpulseNorm = DistSq > 0.f ? Delta.GetUnsafeNormal() : FVector(0.f, 0.f, 1.f);
+								const FVector Fudge = FVector(0.f, 0.f, NetState.ActiveEffect.ResponseMagnitude.Z);
+								HitPT->SetV(HitPT->V() + (ImpulseNorm * NetState.ActiveEffect.ResponseMagnitude.X) + Fudge);
+							}							
+						}
+					}
+				}
+
+			}
+			else
+			{
+				// Effect is finished, reset the state
+				NetState.ActiveEffect = FPhysicsEffectDef();
+			}
+		}
+	}
+}
+
+void IPhysicsEffectsInterface::InitializePhysicsEffects(FPhysicsEffectLocalState&& LocalState)
+{
+	FNetworkPredictionAsyncProxy& NetworkPredictionProxy = GetNetworkPredictionAsyncProxy();
+	FPhysicsEffectsExternalState& State = GetPhysicsEffectsExternalState();
+
+	NetworkPredictionProxy.RegisterSim<FPhysicsEffectsAsyncModelDef>(MoveTemp(LocalState), FPhysicsEffectNetState(), &State.PendingEffectCmd, &State.PhysicsEffectState);
+}
+
+void IPhysicsEffectsInterface::SyncPhysicsEffects()
+{
+	FNetworkPredictionAsyncProxy& NetworkPredictionProxy = GetNetworkPredictionAsyncProxy();
+	FPhysicsEffectsExternalState& State = GetPhysicsEffectsExternalState();
+
+	// This is a crude way to determine when a PhysicsEffect has been played and its results have first been detected on the GT
+	// This covers both GT->PT latency and Remote client replication. This lets us tie cosmetic effects to Physics Effects.
+	// The "NetSimCue" system in Network Prediction can replace this with something more robust.
+	const int32 StartFrame = State.PhysicsEffectState.ActiveEffect.StartFrame;
+	if (StartFrame != INDEX_NONE && StartFrame != State.CachedWindUpSimFrame)
+	{
+		UpdatePendingWindUp(StartFrame, State.PhysicsEffectState.ActiveEffect.TypeID);
+	}
+
+	if (State.PhysicsEffectState.LastPlayedSimFrame != State.CachedLastPlayedSimFrame)
+	{
+		State.CachedLastPlayedSimFrame = State.PhysicsEffectState.LastPlayedSimFrame;
+		OnPhysicsEffectExecuted(State.PhysicsEffectState.LastPlayedEffectTypeID);
+	}
+}
+
+void IPhysicsEffectsInterface::UpdatePendingWindUp(int32 Frame, uint8 TypeID)
+{
+	FPhysicsEffectsExternalState& State = GetPhysicsEffectsExternalState();
+
+	if (Frame > State.CachedWindUpSimFrame)
+	{
+		State.CachedWindUpSimFrame = Frame;
+		OnPhysicsEffectWindUp(TypeID);
+	}
+}
+
+void IPhysicsEffectsInterface::ApplyPhysicsEffectDef(FPhysicsEffectDef&& Effect, const FPhysicsEffectParams& Params)
+{
+	FNetworkPredictionAsyncProxy& NetworkPredictionProxy = GetNetworkPredictionAsyncProxy();
+	FPhysicsEffectsExternalState& State = GetPhysicsEffectsExternalState();
+
+	const int32 DelayFrames = Params.DelaySeconds > 0.f ? (Params.DelaySeconds * 30.f) : 0;
+	Effect.StartFrame = NetworkPredictionProxy.GetNextSimFrame() + PhysicsEffectsDefaultFrameDelay + DelayFrames;
+	
+	if (Params.DurationSeconds >= 0.f)
+	{
+		Effect.EndFrame = Effect.StartFrame + (Params.DurationSeconds * 30.f); // Fixme: what is fast way to access this without having to cache it off locally/per instance (it can't change at runtime but is config var)
+	}
+	else
+	{
+		Effect.EndFrame = -1;
+	}
+
+	if (IsController())
+	{
+		State.PendingEffectCmd.PendingEffect = MoveTemp(Effect);
+	}
+	else
+	{
+		NetworkPredictionProxy.ModifyNetState<FPhysicsEffectsAsyncModelDef>([CapturedEffect = MoveTemp(Effect)](FPhysicsEffectNetState& NetState)
+		{
+			NetState.ActiveEffect = CapturedEffect;
+		});
+	}
+
+	UpdatePendingWindUp(Effect.StartFrame, Effect.TypeID);
+}
+
+FPhysicsEffectNetState IPhysicsEffectsInterface::Debug_ReadActivePhysicsEffect() const
+{
+	const FPhysicsEffectsExternalState& State = const_cast<IPhysicsEffectsInterface*>(this)->GetPhysicsEffectsExternalState();
+	return State.PhysicsEffectState;
+}
+
+int32 IPhysicsEffectsInterface::Debug_GetLatestSimFrame() const
+{
+	FNetworkPredictionAsyncProxy& NetworkPredictionProxy = const_cast<IPhysicsEffectsInterface*>(this)->GetNetworkPredictionAsyncProxy();
+	return NetworkPredictionProxy.GetLatestOutputSimFrame();
+}
+
+uint8 IPhysicsEffectsInterface::GetActivePhysicsEffectTypeID() const
+{
+	const FNetworkPredictionAsyncProxy& NetworkPredictionProxy = const_cast<IPhysicsEffectsInterface*>(this)->GetNetworkPredictionAsyncProxy();
+	const FPhysicsEffectsExternalState& State = const_cast<IPhysicsEffectsInterface*>(this)->GetPhysicsEffectsExternalState();
+
+	const int32 StartFrame = State.PhysicsEffectState.ActiveEffect.StartFrame;
+	const int32 EndFrame = State.PhysicsEffectState.ActiveEffect.EndFrame;
+	const int32 Frame = NetworkPredictionProxy.GetLatestOutputSimFrame();
+
+	if (StartFrame != INDEX_NONE && Frame >= StartFrame && (Frame <= EndFrame || EndFrame == INDEX_NONE))
+	{
+		return State.PhysicsEffectState.LastPlayedEffectTypeID;
+	}
+	return 0;
+}
+
+
+void IPhysicsEffectsInterface::ClearPhysicsEffects()
+{
+	FNetworkPredictionAsyncProxy& NetworkPredictionProxy = GetNetworkPredictionAsyncProxy();
+	FPhysicsEffectsExternalState& State = GetPhysicsEffectsExternalState();
+
+	if (IsController())
+	{
+		// Need to always "advance" StartFrame on the InputCmd to get it to be picked up in Tick_Internal.
+		State.PendingEffectCmd.PendingEffect.StartFrame = FMath::Max(State.PendingEffectCmd.PendingEffect.StartFrame+1, NetworkPredictionProxy.GetNextSimFrame());
+		State.PendingEffectCmd.PendingEffect.EndFrame = 0;
+	}
+	else
+	{
+		NetworkPredictionProxy.ModifyNetState<FPhysicsEffectsAsyncModelDef>([](FPhysicsEffectNetState& NetState)
+		{
+			NetState.ActiveEffect.EndFrame = 0;
+		});
+	}
+}
+
+void IPhysicsEffectsInterface::ApplyImpulse(float Radius, float Magnitude, FTransform RelativeOffset, float FudgeVelocityZ, const FPhysicsEffectParams& Params)
+{
+	FPhysicsEffectDef NewEffect;
+
+	NewEffect.QueryType = EPhysicsEffectQueryType::SphereOverlap;
+	NewEffect.QueryData = FVector(Radius, 0.f, 0.f);
+	NewEffect.QueryRelativeTransform = RelativeOffset;
+	
+	NewEffect.ResponseType = EPhysicsEffectResponseType::Impulse;
+	NewEffect.ResponseMagnitude = FVector(Magnitude, 0.f, FudgeVelocityZ);
+
+	NewEffect.TypeID = Params.TypeID;
+
+	ApplyPhysicsEffectDef(MoveTemp(NewEffect), Params);
+}
+
diff --git a/Engine/Plugins/Runtime/NetworkPredictionExtras/Source/NetworkPredictionExtras/Private/PhysicsEffectsComponent.cpp b/Engine/Plugins/Runtime/NetworkPredictionExtras/Source/NetworkPredictionExtras/Private/PhysicsEffectsComponent.cpp
new file mode 100644
--- /dev/null
+++ b/Engine/Plugins/Runtime/NetworkPredictionExtras/Source/NetworkPredictionExtras/Private/PhysicsEffectsComponent.cpp
@@ -0,0 +1,71 @@
+// Copyright Epic Games, Inc. All Rights Reserved.
+
+#include "PhysicsEffectsComponent.h"
+#include "Net/UnrealNetwork.h"
+
+UPhysicsEffectsComponent::UPhysicsEffectsComponent()
+{
+	bWantsInitializeComponent = true;
+	bAutoActivate = true;
+	SetIsReplicatedByDefault(true);
+
+	PrimaryComponentTick.bCanEverTick = true;
+	PrimaryComponentTick.TickGroup = TG_DuringPhysics;
+	PrimaryComponentTick.bStartWithTickEnabled = true;
+}
+
+void UPhysicsEffectsComponent::InitializeComponent()
+{
+	Super::InitializeComponent();
+#if WITH_CHAOS
+	UWorld* World = GetWorld();	
+	checkSlow(World);
+	
+	FPhysicsEffectLocalState LocalState;
+	LocalState.Proxy = this->GetManagedProxy();
+	LocalState.QueryParams.AddIgnoredActor(GetOwner());
+		
+	if (LocalState.Proxy)
+	{
+		if (ensure(NetworkPredictionProxy.RegisterProxy(GetWorld())))
+		{
+			InitializePhysicsEffects(MoveTemp(LocalState));
+		}
+		else
+		{
+			UE_LOG(LogNetworkPhysics, Warning, TEXT("No valid physics body found on %s"), *GetName());
+		}
+	}
+#endif
+}
+
+void UPhysicsEffectsComponent::UninitializeComponent()
+{
+	Super::UninitializeComponent();
+	NetworkPredictionProxy.UnregisterProxy();
+}
+
+void UPhysicsEffectsComponent::TickComponent(float DeltaTime, enum ELevelTick TickType, FActorComponentTickFunction *ThisTickFunction)
+{
+	NetworkPredictionProxy.RegisterController(nullptr); // fixme this is needed but shouldn't be
+	SyncPhysicsEffects();
+	Super::TickComponent(DeltaTime, TickType, ThisTickFunction);
+}
+
+void UPhysicsEffectsComponent::PreReplication(IRepChangedPropertyTracker & ChangedPropertyTracker)
+{
+	Super::PreReplication(ChangedPropertyTracker);
+	NetworkPredictionProxy.OnPreReplication();
+}
+
+void UPhysicsEffectsComponent::GetLifetimeReplicatedProps(TArray<FLifetimeProperty>& OutLifetimeProps) const
+{
+	Super::GetLifetimeReplicatedProps(OutLifetimeProps);
+	DOREPLIFETIME( UPhysicsEffectsComponent, NetworkPredictionProxy);	
+}
+
+bool UPhysicsEffectsComponent::IsController() const
+{
+	const bool bIServer = GetOwnerRole() == ROLE_Authority;
+	return bIServer;
+}
\ No newline at end of file
diff --git a/Engine/Plugins/Runtime/NetworkPredictionExtras/Source/NetworkPredictionExtras/Private/PhysicsMovement.cpp b/Engine/Plugins/Runtime/NetworkPredictionExtras/Source/NetworkPredictionExtras/Private/PhysicsMovement.cpp
--- a/Engine/Plugins/Runtime/NetworkPredictionExtras/Source/NetworkPredictionExtras/Private/PhysicsMovement.cpp
+++ b/Engine/Plugins/Runtime/NetworkPredictionExtras/Source/NetworkPredictionExtras/Private/PhysicsMovement.cpp
@@ -23,6 +23,7 @@
 #include "Async/NetworkPredictionAsyncModelDefRegistry.h"
 #include "Async/NetworkPredictionAsyncProxyImpl.h"
 #include "Components/PrimitiveComponent.h"
+#include "PhysicsEffect.h"
 
 
 NETSIM_DEVCVAR_SHIPCONST_INT(DisablePhysicsMovement, 0, "np2.DisablePhysicsMovement", "");
@@ -420,19 +421,29 @@ void UPhysicsMovementComponent::InitializeComponent()
 	UWorld* World = GetWorld();	
 	checkSlow(World);
 	
-	FPhysicsMovementLocalState LocalState;
-	LocalState.Proxy = this->GetManagedProxy();
-	LocalState.QueryParams.AddIgnoredActor(GetOwner());
-	if (LocalState.Proxy)
 	{
-		if (ensure(NetworkPredictionProxy.RegisterProxy(GetWorld())))
+		FPhysicsMovementLocalState LocalState;
+		LocalState.Proxy = this->GetManagedProxy();
+		LocalState.QueryParams.AddIgnoredActor(GetOwner());
+		if (LocalState.Proxy)
 		{
-			NetworkPredictionProxy.RegisterSim<FPhysicsMovementAsyncModelDef>(MoveTemp(LocalState), FPhysicsMovementNetState(), &PendingInputCmd, &MovementState);
+			if (ensure(NetworkPredictionProxy.RegisterProxy(GetWorld())))
+			{
+				NetworkPredictionProxy.RegisterSim<FPhysicsMovementAsyncModelDef>(MoveTemp(LocalState), FPhysicsMovementNetState(), &PendingInputCmd, &MovementState);
+
+				if (bEnablePhysicsEffects)
+				{
+					FPhysicsEffectLocalState PELocalState;
+					PELocalState.Proxy = this->GetManagedProxy();
+					PELocalState.QueryParams.AddIgnoredActor(GetOwner());
+					InitializePhysicsEffects(MoveTemp(PELocalState));
+				}
+			}
+		}
+		else
+		{
+			UE_LOG(LogNetworkPhysics, Warning, TEXT("No valid physics body found on %s"), *GetName());
 		}
-	}
-	else
-	{
-		UE_LOG(LogNetworkPhysics, Warning, TEXT("No valid physics body found on %s"), *GetName());
 	}
 #endif
 }
@@ -451,6 +462,8 @@ void UPhysicsMovementComponent::UninitializeComponent()
 
 void UPhysicsMovementComponent::TickComponent(float DeltaTime, enum ELevelTick TickType, FActorComponentTickFunction *ThisTickFunction)
 {
+	SyncPhysicsEffects();
+
 	Super::TickComponent(DeltaTime, TickType, ThisTickFunction);
 	
 	APlayerController* PC = GetOwnerPC();
@@ -470,6 +483,22 @@ void UPhysicsMovementComponent::TickComponent(float DeltaTime, enum ELevelTick T
 		PendingInputCmd.bLegit = true;
 		PendingInputCmd.Counter++;
 	}
+
+
+#if 0 // Requires SetInterpChannel
+	if ((GetOwnerRole() != ROLE_Authority))
+	{
+		if (Chaos::FSingleParticlePhysicsProxy* Proxy = this->GetManagedProxy())
+		{
+			const int32 InterpChannel = PC == nullptr ? 1 : 0;
+			if (Proxy->GetInterpChannel() != InterpChannel)
+			{
+				Proxy->SetInterpChannel(InterpChannel);
+				UE_LOG(LogTemp, Warning, TEXT("%s SetInterpChannel %d"), *GetName(), InterpChannel);
+			}
+		}
+	}
+#endif
 }
 
 void UPhysicsMovementComponent::GetLifetimeReplicatedProps(TArray<FLifetimeProperty>& OutLifetimeProps) const
@@ -573,23 +602,45 @@ void UPhysicsMovementComponent::TestMisprediction()
 	});	
 }
 
+bool UPhysicsMovementComponent::IsController() const
+{
+	APlayerController* PC = GetOwnerPC();
+	if (PC && PC->IsLocalController())
+	{
+		return true;
+	}
+
+	const bool bIServer = GetOwnerRole() == ROLE_Authority;
+	if (bIServer && PC == nullptr)
+	{
+		return true;
+	}
+
+	return false;
+}
+
 FAutoConsoleCommandWithWorldAndArgs ForcePhysicsMovementCorrectionCmd(TEXT("np.PhysicsMovement.ForceRandCorrection2"), TEXT(""),
 FConsoleCommandWithWorldAndArgsDelegate::CreateLambda([](const TArray< FString >& Args, UWorld* InWorld) 
+{
+	for (TObjectIterator<UWorld> It; It; ++It)
 	{
-		for (TObjectIterator<UWorld> It; It; ++It)
+		UWorld* World = *It;
+		if (World->GetNetMode() != NM_Client && (World->WorldType == EWorldType::PIE || World->WorldType == EWorldType::Game))
 		{
-			UWorld* World = *It;
-			if (World->GetNetMode() != NM_Client && (World->WorldType == EWorldType::PIE || World->WorldType == EWorldType::Game))
-			{
-				const int32 NewRand = FMath::RandHelper(1024);
+			const int32 NewRand = FMath::RandHelper(1024);
 
-				for (TActorIterator<AActor> ActorIt(World); ActorIt; ++ActorIt)
+			for (TActorIterator<AActor> ActorIt(World); ActorIt; ++ActorIt)
+			{
+				if (UPhysicsMovementComponent* PhysComp = ActorIt->FindComponentByClass<UPhysicsMovementComponent>())
 				{
-					if (UPhysicsMovementComponent* PhysComp = ActorIt->FindComponentByClass<UPhysicsMovementComponent>())
-					{
-						PhysComp->TestMisprediction();
-					}
+					PhysComp->TestMisprediction();
 				}
 			}
 		}
-	}));
+	}
+}));
+
+
+
+
+
diff --git a/Engine/Plugins/Runtime/NetworkPredictionExtras/Source/NetworkPredictionExtras/Public/PhysicsEffect.h b/Engine/Plugins/Runtime/NetworkPredictionExtras/Source/NetworkPredictionExtras/Public/PhysicsEffect.h
new file mode 100644
--- /dev/null
+++ b/Engine/Plugins/Runtime/NetworkPredictionExtras/Source/NetworkPredictionExtras/Public/PhysicsEffect.h
@@ -0,0 +1,126 @@
+// Copyright Epic Games, Inc. All Rights Reserved.
+
+#pragma once
+
+#include "CoreTypes.h"
+#include "CoreFwd.h"
+#include "Async/NetworkPredictionAsyncModelDef.h"
+#include "Async/NetworkPredictionAsyncID.h"
+#include "PhysicsEffectTypes.h"
+#include "PhysicsEffect.generated.h"
+
+
+// Physics Effects is an NP system for doing generic physics effects like radial impulses.
+// The goal is to provide a simple API for doing things on the PT without having to modify code.
+// This is somewhat of a stop gap until a PT blueprint is available.
+//
+// The key here is that we are *computing* what forces/impulses/etc to apply on the PT.
+// We already have a way to marshal absolute forces/impulses/etc from the GT but this is usually not what we want - we want to calculate these magnitudes based on the most accurate physics data/scene queries.
+// Right now this is brutually rough and stubbed in. Will build up a more robust API in FPhysicsEffectDef.
+//
+// Note on the UE4 glue:
+//	-Goal is to have the shared API on IPhysicsEffectsInterface.
+//	-UPhysicsEffectComponent can be used to add this interface to any actor (e.g, a grenade)
+//	-UPhysicsMovementComponent also implements IPhysicsEffectsInterface.
+//	-We just don't want to put this on *every* physics object since their is additional overhead.
+
+
+struct FNetworkPredictionAsyncProxy;
+
+// State tracked on the GT for PEs
+struct FPhysicsEffectsExternalState
+{
+	FPhysicsEffectCmd PendingEffectCmd;
+	FPhysicsEffectNetState PhysicsEffectState;
+	int32 CachedLastPlayedSimFrame = 0;
+	int32 CachedWindUpSimFrame = 0;
+};
+
+USTRUCT(BlueprintType, meta=(ShowOnlyInnerProperties))
+struct FPhysicsEffectParams
+{
+	GENERATED_BODY()
+
+	// Length of time to apply this effect from when it starts. 0=single frame. Negative number for infinite duration. (Use ClearPhysicsEffect to stop it)
+	UPROPERTY(BlueprintreadWrite, Category="Physics Effect")
+	float DurationSeconds=0.0f;
+
+	// Delay before starting the effect. Even a small amount of delay can help hide network latency to clients.
+	UPROPERTY(BlueprintreadWrite, Category="Physics Effect")
+	float DelaySeconds=0.f;
+
+	UPROPERTY(BlueprintreadWrite, Category="Physics Effect")
+	uint8 TypeID=0;
+};
+
+// Forward input producing event to someone else (probably the owning actor)
+DECLARE_DYNAMIC_MULTICAST_DELEGATE_OneParam(FOnPhysicsEffects, int32, TypeID);
+
+// UE Interface for Physics Effects
+UINTERFACE(MinimalAPI, meta = (CannotImplementInterfaceInBlueprint))
+class UPhysicsEffectsInterface : public UInterface
+{
+	GENERATED_BODY()
+};
+
+class NETWORKPREDICTIONEXTRAS_API IPhysicsEffectsInterface
+{
+	GENERATED_BODY()
+
+public:
+
+	// Return the NP AsyncProxy to use for PhysicsEffects
+	virtual FNetworkPredictionAsyncProxy& GetNetworkPredictionAsyncProxy() = 0;
+
+	// Return the GT state for PhysicsEffects
+	virtual FPhysicsEffectsExternalState& GetPhysicsEffectsExternalState() = 0;
+
+	// Whether we are the controller of the PE system instance. Determines if we can write to InputCmds or do direct NetState mods.
+	virtual bool IsController() const = 0;
+
+	// Notify that a PE has executed (on the PT) and its output/side effects/(whatever it did) has been marshaled back (to the GT)
+	virtual void OnPhysicsEffectExecuted(uint8 TypeID) { };
+
+	// Notify a PE is pending execute but hasn't been marshaled back to the GT yet
+	virtual void OnPhysicsEffectWindUp(uint8 TypeID) { };
+
+	// Call in InitializeComponent() or anytime after the NP Async Proxy has been registered
+	void InitializePhysicsEffects(FPhysicsEffectLocalState&& LocalState);
+
+	// Call to sync with PT.
+	void SyncPhysicsEffects();
+
+	// ---------------------------------------------------------------------
+	//	Temp/Debugging API
+	// ---------------------------------------------------------------------
+
+	// Returns latest copy of the active physics effect state. This is intended for debugging purposes. Function may be removed in the future.
+	UFUNCTION(BlueprintCallable, Category="Physics Effect")
+	virtual FPhysicsEffectNetState Debug_ReadActivePhysicsEffect() const;
+
+	// Gets latest output sim frame. Debug only. Function may be removed in the future.
+	UFUNCTION(BlueprintCallable, Category="Physics Effect")
+	virtual int32 Debug_GetLatestSimFrame() const;
+
+	// ---------------------------------------------------------------------
+	//	Physics Effects API
+	// ---------------------------------------------------------------------
+
+	// Get Active Physics Effect TypeID. Can be used to determine "what is playing" (0 if nothing)
+	UFUNCTION(BlueprintCallable, Category="Physics Effect")
+	virtual uint8 GetActivePhysicsEffectTypeID() const;
+
+	// Clears active PhysicsEffect
+	UFUNCTION(BlueprintCallable, Category="Physics Effect")
+	virtual void ClearPhysicsEffects();
+
+	UFUNCTION(BlueprintCallable, Category="Physics Effect")
+	virtual void ApplyImpulse(float Radius, float Magnitude, FTransform RelativeOffset, float FudgeVelocityZ, const FPhysicsEffectParams& Params);
+
+private:
+
+	void UpdatePendingWindUp(int32 Frame, uint8 TypeID);
+	void ApplyPhysicsEffectDef(FPhysicsEffectDef&& Effect, const FPhysicsEffectParams& Params);
+};
+
+
diff --git a/Engine/Plugins/Runtime/NetworkPredictionExtras/Source/NetworkPredictionExtras/Public/PhysicsEffectTypes.h b/Engine/Plugins/Runtime/NetworkPredictionExtras/Source/NetworkPredictionExtras/Public/PhysicsEffectTypes.h
new file mode 100644
--- /dev/null
+++ b/Engine/Plugins/Runtime/NetworkPredictionExtras/Source/NetworkPredictionExtras/Public/PhysicsEffectTypes.h
@@ -0,0 +1,169 @@
+// Copyright Epic Games, Inc. All Rights Reserved.
+
+#pragma once
+
+#include "NetworkPhysics.h"
+#include "PhysicsEffectTypes.generated.h"
+
+namespace Chaos
+{
+	class FSingleParticlePhysicsProxy;
+}
+
+UENUM()
+enum class EPhysicsEffectQueryType : uint8
+{
+	None,	// Don't do a query. The response is applied directly to the target particles instead.
+	SphereOverlap // sphere overlap
+};
+
+UENUM()
+enum class EPhysicsEffectResponseType : uint8
+{
+	Force,
+	Impulse,
+	Torque
+};
+
+USTRUCT(BlueprintType)
+struct FPhysicsEffectDef
+{
+	GENERATED_BODY()
+
+	FPhysicsEffectDef()
+		: QueryData(ForceInitToZero)
+		, ResponseMagnitude(ForceInitToZero)
+	{ }
+
+	UPROPERTY(BlueprintReadWrite, Category="Physics Effect")
+	int32 StartFrame = INDEX_NONE;
+
+	UPROPERTY(BlueprintReadWrite, Category="Physics Effect")
+	int32 EndFrame = INDEX_NONE;
+
+	UPROPERTY(BlueprintReadWrite, Category="Physics Effect")
+	EPhysicsEffectQueryType QueryType = EPhysicsEffectQueryType::SphereOverlap;
+
+	UPROPERTY(BlueprintReadWrite, Category="Physics Effect")
+	FVector QueryData;
+
+	UPROPERTY(BlueprintReadWrite, Category="Physics Effect")
+	FTransform QueryRelativeTransform = FTransform::Identity;
+
+	UPROPERTY(BlueprintReadWrite, Category="Physics Effect")
+	EPhysicsEffectResponseType ResponseType = EPhysicsEffectResponseType::Impulse;
+
+	UPROPERTY(BlueprintReadWrite, Category="Physics Effect")
+	FVector ResponseMagnitude;
+
+	UPROPERTY(BlueprintReadWrite, Category="Physics Effect")
+	uint8 TypeID=0;
+
+	void NetSerialize(FArchive& Ar)
+	{
+		Ar << StartFrame;
+		if (StartFrame != INDEX_NONE)
+		{
+			Ar << EndFrame;
+			Ar << QueryType;
+			Ar << QueryData;
+			Ar << QueryRelativeTransform;
+			Ar << ResponseType;
+			Ar << ResponseMagnitude;
+			Ar << TypeID;
+		}
+	}
+
+	bool ShouldReconcile(const FPhysicsEffectDef& AuthState) const
+	{
+		if (StartFrame == INDEX_NONE && StartFrame == AuthState.StartFrame)
+		{
+			// If start frame is -1 and we both agree, don't worry about the rest of the payload
+			return false;
+		}
+
+		return 
+			StartFrame != AuthState.StartFrame ||
+			EndFrame != AuthState.EndFrame ||
+			QueryType != AuthState.QueryType ||
+			QueryData != AuthState.QueryData ||
+			!QueryRelativeTransform.Equals(AuthState.QueryRelativeTransform) ||
+			ResponseType != AuthState.ResponseType ||
+			ResponseMagnitude != AuthState.ResponseMagnitude ||
+			TypeID != AuthState.TypeID;
+	}
+
+	FString ToString() const
+	{
+		return FString::Printf(TEXT("{StartFrame: %d EndFrame: %d QueryType: %d QueryData: %s ResponseType: %d ResponseMagnitude: %s"),
+			StartFrame, EndFrame, QueryType, *QueryData.ToString(), ResponseType, *ResponseMagnitude.ToString());
+	}
+};
+
+USTRUCT(BlueprintType)
+struct FPhysicsEffectCmd
+{
+	GENERATED_BODY()
+
+	UPROPERTY(BlueprintReadWrite, Category="Physics Effect")
+	FPhysicsEffectDef PendingEffect;
+
+	void NetSerialize(FArchive& Ar)
+	{
+		PendingEffect.NetSerialize(Ar);
+	}
+
+	bool ShouldReconcile(const FPhysicsEffectCmd& AuthState) const
+	{
+		return PendingEffect.ShouldReconcile(AuthState.PendingEffect);
+	}
+
+	FString ToString() const
+	{
+		return PendingEffect.ToString();
+	}
+};
+
+USTRUCT(BlueprintType)
+struct FPhysicsEffectNetState
+{
+	GENERATED_BODY()
+
+	UPROPERTY(BlueprintReadWrite, Category="Physics Effect")
+	FPhysicsEffectDef ActiveEffect;
+
+	UPROPERTY(BlueprintReadWrite, Category="Physics Effect")
+	uint8 LastPlayedEffectTypeID = 0;
+
+	UPROPERTY(BlueprintReadWrite, Category="Physics Effect")
+	int32 LastPlayedSimFrame = INDEX_NONE;
+
+	void NetSerialize(FArchive& Ar)
+	{
+		ActiveEffect.NetSerialize(Ar);
+		Ar << LastPlayedEffectTypeID;
+		Ar << LastPlayedSimFrame;
+	}
+
+	bool ShouldReconcile(const FPhysicsEffectNetState& AuthState) const
+	{
+		return ActiveEffect.ShouldReconcile(AuthState.ActiveEffect) ||
+			LastPlayedEffectTypeID != AuthState.LastPlayedEffectTypeID ||
+			LastPlayedSimFrame != AuthState.LastPlayedSimFrame;
+	}
+
+	FString ToString() const
+	{
+		return ActiveEffect.ToString();
+	}
+};
+
+USTRUCT(BlueprintType)
+struct FPhysicsEffectLocalState
+{
+	GENERATED_BODY()
+
+	Chaos::FSingleParticlePhysicsProxy* Proxy = nullptr;
+
+	FCollisionQueryParams QueryParams = FCollisionQueryParams::DefaultQueryParam;
+};
\ No newline at end of file
diff --git a/Engine/Plugins/Runtime/NetworkPredictionExtras/Source/NetworkPredictionExtras/Public/PhysicsEffectsComponent.h b/Engine/Plugins/Runtime/NetworkPredictionExtras/Source/NetworkPredictionExtras/Public/PhysicsEffectsComponent.h
new file mode 100644
--- /dev/null
+++ b/Engine/Plugins/Runtime/NetworkPredictionExtras/Source/NetworkPredictionExtras/Public/PhysicsEffectsComponent.h
@@ -0,0 +1,50 @@
+// Copyright Epic Games, Inc. All Rights Reserved.
+
+#pragma once
+
+#include "NetworkPhysicsComponent.h"
+#include "PhysicsEffect.h"
+#include "Async/NetworkPredictionAsyncProxy.h"
+#include "PhysicsEffectsComponent.generated.h"
+
+// Networked Physics Component that can apply PhysicsEffects. This is intended to be a minimal implementation of the IPhysicsEffectsInterface.
+// This common use case would be for objects that are not using PhysicsMovement but still need to apply PEs. Examples: grenades, rockets, etc.
+UCLASS(BlueprintType, meta=(BlueprintSpawnableComponent))
+class NETWORKPREDICTIONEXTRAS_API UPhysicsEffectsComponent : public UNetworkPhysicsComponent, public IPhysicsEffectsInterface
+{
+	GENERATED_BODY()
+
+public:
+
+	UPhysicsEffectsComponent();
+
+	virtual void InitializeComponent() override;
+	virtual void UninitializeComponent() override;
+	virtual void TickComponent(float DeltaTime, enum ELevelTick TickType, FActorComponentTickFunction *ThisTickFunction) override;
+	virtual void PreReplication(IRepChangedPropertyTracker & ChangedPropertyTracker);
+
+	UPROPERTY(Replicated)
+	FNetworkPredictionAsyncProxy NetworkPredictionProxy;
+
+	FNetworkPredictionAsyncProxy& GetNetworkPredictionAsyncProxy() override { return NetworkPredictionProxy; }
+	FPhysicsEffectsExternalState& GetPhysicsEffectsExternalState() override { return PhysicsEffectsExternalState; }
+
+	bool IsController() const override;
+
+	// Forward input producing event to someone else (probably the owning actor)
+	DECLARE_DYNAMIC_MULTICAST_DELEGATE_OneParam(FOnPhysicsEffect, int32, TypeID);
+	
+	// Delegate to inform the GT/bp that a PE has taken place and is now "seen" on the GT
+	UPROPERTY(BlueprintAssignable, Category = "Physics Effect")
+	FOnPhysicsEffect NotifyPhysicsEffectExecuted;
+	void OnPhysicsEffectExecuted(uint8 TypeID) override { NotifyPhysicsEffectExecuted.Broadcast(TypeID); };
+
+	// Delegate to inform the GT/bp  that a PE will take place soon but not hasn't happened yet
+	UPROPERTY(BlueprintAssignable, Category = "Physics Effect")
+	FOnPhysicsEffect NotifyPhysicsEffectWindUp;
+	void OnPhysicsEffectWindUp(uint8 TypeID) override { NotifyPhysicsEffectWindUp.Broadcast(TypeID); };
+
+private:
+
+	FPhysicsEffectsExternalState PhysicsEffectsExternalState;
+};
\ No newline at end of file
diff --git a/Engine/Plugins/Runtime/NetworkPredictionExtras/Source/NetworkPredictionExtras/Public/PhysicsMovement.h b/Engine/Plugins/Runtime/NetworkPredictionExtras/Source/NetworkPredictionExtras/Public/PhysicsMovement.h
--- a/Engine/Plugins/Runtime/NetworkPredictionExtras/Source/NetworkPredictionExtras/Public/PhysicsMovement.h
+++ b/Engine/Plugins/Runtime/NetworkPredictionExtras/Source/NetworkPredictionExtras/Public/PhysicsMovement.h
@@ -8,6 +8,7 @@
 #include "NetworkPhysicsComponent.h"
 #include "NetworkPredictionCVars.h"
 #include "Springs.h"
+#include "PhysicsEffect.h"
 #include "PhysicsMovement.generated.h"
 
 namespace UE_NP
@@ -204,6 +205,11 @@ struct FPhysicsMovementNetState
 			Springs != AuthState.Springs || 
 			(UE_NP::ForceReconcilePhysicsMovementState() > 0);
 	}
+
+	FString ToString() const
+	{
+		return FString();
+	}
 };
 
 USTRUCT(BlueprintType)
@@ -217,7 +223,7 @@ struct FPhysicsMovementLocalState
 };
 
 UCLASS(BlueprintType, meta=(BlueprintSpawnableComponent))
-class NETWORKPREDICTIONEXTRAS_API UPhysicsMovementComponent : public UNetworkPhysicsComponent
+class NETWORKPREDICTIONEXTRAS_API UPhysicsMovementComponent : public UNetworkPhysicsComponent, public IPhysicsEffectsInterface
 {
 	GENERATED_BODY()
 
@@ -247,7 +253,11 @@ public:
 	UPROPERTY(BlueprintAssignable, Category = "Input")
 	FOnGenerateLocalInputCmd OnGenerateLocalInputCmd;
 
+	void TestMisprediction();
 	
+	// ---------------------------------------------------------------------
+	//	Movement
+	// ---------------------------------------------------------------------
 
 	UFUNCTION(BlueprintCallable, Category="Movement")
 	void SetAutoTargetYawStrength(float Strength);
@@ -277,10 +287,36 @@ public:
 	UFUNCTION(BlueprintCallable, Category = "Movement")
 	void ClearSprings();
 
-	void TestMisprediction();
+
+	// ---------------------------------------------------------------------
+	//	Physics Effects
+	// ---------------------------------------------------------------------
+
+	UPROPERTY(EditDefaultsOnly, Category="Physics Effects")
+	bool bEnablePhysicsEffects=true;
+
+	FNetworkPredictionAsyncProxy& GetNetworkPredictionAsyncProxy() override { return NetworkPredictionProxy; }
+	FPhysicsEffectsExternalState& GetPhysicsEffectsExternalState() override { return PhysicsEffectsExternalState; }
+
+	bool IsController() const override;
+
+	// Forward input producing event to someone else (probably the owning actor)
+	DECLARE_DYNAMIC_MULTICAST_DELEGATE_OneParam(FOnPhysicsEffect, int32, TypeID);
+
+	// Delegate to inform the GT/bp that a PE has taken place and is now "seen" on the GT
+	UPROPERTY(BlueprintAssignable, Category = "Physics Effect")
+	FOnPhysicsEffect NotifyPhysicsEffectExecuted;
+	void OnPhysicsEffectExecuted(uint8 TypeID) override { NotifyPhysicsEffectExecuted.Broadcast(TypeID); };
+
+	// Delegate to inform the GT/bp  that a PE will take place soon but not hasn't happened yet
+	UPROPERTY(BlueprintAssignable, Category = "Physics Effect")
+	FOnPhysicsEffect NotifyPhysicsEffectWindUp;
+	void OnPhysicsEffectWindUp(uint8 TypeID) override { NotifyPhysicsEffectWindUp.Broadcast(TypeID); };
 
 private:
 
+	FPhysicsEffectsExternalState PhysicsEffectsExternalState;
+
 	APlayerController* GetOwnerPC() const;
 
 	// Crude way of detecting possession change
