Network Physics
-Set remote client interp channel to 1

diff --git a/Engine/Plugins/Runtime/NetworkPredictionExtras/Source/NetworkPredictionExtras/Private/PhysicsMovement.cpp b/Engine/Plugins/Runtime/NetworkPredictionExtras/Source/NetworkPredictionExtras/Private/PhysicsMovement.cpp
--- a/Engine/Plugins/Runtime/NetworkPredictionExtras/Source/NetworkPredictionExtras/Private/PhysicsMovement.cpp
+++ b/Engine/Plugins/Runtime/NetworkPredictionExtras/Source/NetworkPredictionExtras/Private/PhysicsMovement.cpp
@@ -485,20 +485,20 @@ void UPhysicsMovementComponent::TickComponent(float DeltaTime, enum ELevelTick T
 	}
 
 
-#if 0 // Requires SetInterpChannel
+	// Fixme: this should probably be moved to a more central place or at least wrapped in a utility function.
+	// It is not specific to PhysicsMovement, it is just saying "put remote controlled movers in the 2nd interp channel"
 	if ((GetOwnerRole() != ROLE_Authority))
 	{
 		if (Chaos::FSingleParticlePhysicsProxy* Proxy = this->GetManagedProxy())
 		{
 			const int32 InterpChannel = PC == nullptr ? 1 : 0;
-			if (Proxy->GetInterpChannel() != InterpChannel)
+			if (Proxy->GetInterpChannel_External() != InterpChannel)
 			{
-				Proxy->SetInterpChannel(InterpChannel);
-				UE_LOG(LogTemp, Warning, TEXT("%s SetInterpChannel %d"), *GetName(), InterpChannel);
+				Proxy->SetInterpChannel_External(InterpChannel);
+				UE_LOG(LogNetworkPrediction, Log, TEXT("%s SetInterpChannel %d"), *GetName(), InterpChannel);
 			}
 		}
 	}
-#endif
 }
 
 void UPhysicsMovementComponent::GetLifetimeReplicatedProps(TArray<FLifetimeProperty>& OutLifetimeProps) const
