Physics Effects
-Fxied some bugs and added callback for windup ending

diff --git a/Engine/Plugins/Runtime/NetworkPredictionExtras/Source/NetworkPredictionExtras/Private/PhysicsEffect.cpp b/Engine/Plugins/Runtime/NetworkPredictionExtras/Source/NetworkPredictionExtras/Private/PhysicsEffect.cpp
--- a/Engine/Plugins/Runtime/NetworkPredictionExtras/Source/NetworkPredictionExtras/Private/PhysicsEffect.cpp
+++ b/Engine/Plugins/Runtime/NetworkPredictionExtras/Source/NetworkPredictionExtras/Private/PhysicsEffect.cpp
@@ -15,6 +15,7 @@
 #include "DrawDebugHelpers.h"
 
 NP_DEVCVAR_INT(PhysicsEffectsDefaultFrameDelay, 5, "np2.pe.FrameDelay", "Default Frame delay for Physics Effects");
+NP_DEVCVAR_INT(PhysicsEffectsDefaultExecuteWindow, 10, "np2.pe.ExecuteWindow", "Consider PEs to be active within this window if they are new to us. Covers low replication rate.");
 
 struct FPhysicsEffectSimulation
 {
@@ -105,12 +106,6 @@ void FPhysicsEffectSimulation::Tick_Internal(FNetworkPredictionAsyncTickContext&
 						}
 					}
 				}
-
-			}
-			else
-			{
-				// Effect is finished, reset the state
-				NetState.ActiveEffect = FPhysicsEffectDef();
 			}
 		}
 	}
@@ -129,63 +124,68 @@ void IPhysicsEffectsInterface::SyncPhysicsEffects()
 	FNetworkPredictionAsyncProxy& NetworkPredictionProxy = GetNetworkPredictionAsyncProxy();
 	FPhysicsEffectsExternalState& State = GetPhysicsEffectsExternalState();
 
-	// This is a crude way to determine when a PhysicsEffect has been played and its results have first been detected on the GT
-	// This covers both GT->PT latency and Remote client replication. This lets us tie cosmetic effects to Physics Effects.
-	// The "NetSimCue" system in Network Prediction can replace this with something more robust.
 	const int32 StartFrame = State.PhysicsEffectState.ActiveEffect.StartFrame;
-	if (StartFrame != INDEX_NONE && StartFrame != State.CachedWindUpSimFrame)
+	const int32 EndFrame = State.PhysicsEffectState.ActiveEffect.EndFrame;
+	const int32 LatestFrame = NetworkPredictionProxy.GetLatestOutputSimFrame();
+	
+	// WindUp
 	{
-		UpdatePendingWindUp(StartFrame, State.PhysicsEffectState.ActiveEffect.TypeID);
-	}
+		uint8 ActiveWindUpTypeID = 0;
+		if ((StartFrame != INDEX_NONE) && (LatestFrame < StartFrame) && (StartFrame >= EndFrame))
+		{
+			ActiveWindUpTypeID = State.PhysicsEffectState.ActiveEffect.TypeID;
+		}
 
-	if (State.PhysicsEffectState.LastPlayedSimFrame != State.CachedLastPlayedSimFrame)
-	{
-		State.CachedLastPlayedSimFrame = State.PhysicsEffectState.LastPlayedSimFrame;
-		OnPhysicsEffectExecuted(State.PhysicsEffectState.LastPlayedEffectTypeID);
-	}
-}
+		if (State.CachedWindUpTypeID != ActiveWindUpTypeID)
+		{
+			if (State.CachedWindUpTypeID != 0)
+			{
+				OnPhysicsEffectWindUpEnd(State.CachedWindUpTypeID);
+			}
 
-void IPhysicsEffectsInterface::UpdatePendingWindUp(int32 Frame, uint8 TypeID)
-{
-	FPhysicsEffectsExternalState& State = GetPhysicsEffectsExternalState();
+			if (ActiveWindUpTypeID != 0)
+			{
+				OnPhysicsEffectWindUp(ActiveWindUpTypeID);
+			}
 
-	if (Frame > State.CachedWindUpSimFrame)
-	{
-		State.CachedWindUpSimFrame = Frame;
-		OnPhysicsEffectWindUp(TypeID);
+			State.CachedWindUpTypeID = ActiveWindUpTypeID;
+		}
 	}
-}
-
-void IPhysicsEffectsInterface::ApplyPhysicsEffectDef(FPhysicsEffectDef&& Effect, const FPhysicsEffectParams& Params)
-{
-	FNetworkPredictionAsyncProxy& NetworkPredictionProxy = GetNetworkPredictionAsyncProxy();
-	FPhysicsEffectsExternalState& State = GetPhysicsEffectsExternalState();
 
-	const int32 DelayFrames = Params.DelaySeconds > 0.f ? (Params.DelaySeconds * 30.f) : 0;
-	Effect.StartFrame = NetworkPredictionProxy.GetNextSimFrame() + PhysicsEffectsDefaultFrameDelay + DelayFrames;
-	
-	if (Params.DurationSeconds >= 0.f)
-	{
-		Effect.EndFrame = Effect.StartFrame + (Params.DurationSeconds * 30.f); // Fixme: what is fast way to access this without having to cache it off locally/per instance (it can't change at runtime but is config var)
-	}
-	else
+	// Active
 	{
-		Effect.EndFrame = -1;
-	}
+		uint8 ActiveTypeID = 0;
+		if ((StartFrame != INDEX_NONE) && (LatestFrame >= StartFrame))
+		{
+			const bool bIsStillPlaying = (LatestFrame <= EndFrame || EndFrame == INDEX_NONE);
+			
+			// We want to detect if an effect is new to us (the GT) but maybe didn't "just happen" due to low replication frequency
+			// (E.,g we don't expect every physics frame results to rep to everyone but as long as you find out about an effect within this window, pretend it "just happened")
+			const bool bNewToUs = State.CachedLastPlayedSimFrame < StartFrame && (LatestFrame <= EndFrame + PhysicsEffectsDefaultExecuteWindow);
+			State.CachedLastPlayedSimFrame = StartFrame;
+
+			if (bIsStillPlaying || bNewToUs)
+			{
+				ActiveTypeID = State.PhysicsEffectState.ActiveEffect.TypeID;
+			}
+		}
 
-	if (IsController())
-	{
-		State.PendingEffectCmd.PendingEffect = MoveTemp(Effect);
-	}
-	else
-	{
-		NetworkPredictionProxy.ModifyNetState<FPhysicsEffectsAsyncModelDef>([CapturedEffect = MoveTemp(Effect)](FPhysicsEffectNetState& NetState)
+		if (ActiveTypeID != State.CachedActiveTypeID)
 		{
-			NetState.ActiveEffect = CapturedEffect;
-		});
-	}
+			if (State.CachedActiveTypeID != 0)
+			{
+				OnPhysicsEffectEnd(State.CachedActiveTypeID);
+			}
+
+			if (ActiveTypeID != 0)
+			{
+				OnPhysicsEffectExecuted(ActiveTypeID);
+			}
 
-	UpdatePendingWindUp(Effect.StartFrame, Effect.TypeID);
+			State.CachedActiveTypeID = ActiveTypeID;
+			State.CachedLastPlayedSimFrame = LatestFrame;
+		}
+	}
 }
 
 FPhysicsEffectNetState IPhysicsEffectsInterface::Debug_ReadActivePhysicsEffect() const
@@ -237,6 +237,36 @@ void IPhysicsEffectsInterface::ClearPhysicsEffects()
 	}
 }
 
+void IPhysicsEffectsInterface::ApplyPhysicsEffectDef(FPhysicsEffectDef&& Effect, const FPhysicsEffectParams& Params)
+{
+	FNetworkPredictionAsyncProxy& NetworkPredictionProxy = GetNetworkPredictionAsyncProxy();
+	FPhysicsEffectsExternalState& State = GetPhysicsEffectsExternalState();
+
+	const int32 DelayFrames = Params.DelaySeconds > 0.f ? (Params.DelaySeconds * 30.f) : 0;
+	Effect.StartFrame = NetworkPredictionProxy.GetNextSimFrame() + PhysicsEffectsDefaultFrameDelay + DelayFrames;
+	
+	if (Params.DurationSeconds >= 0.f)
+	{
+		Effect.EndFrame = Effect.StartFrame + (Params.DurationSeconds * 30.f); // Fixme: what is fast way to access this without having to cache it off locally/per instance (it can't change at runtime but is config var)
+	}
+	else
+	{
+		Effect.EndFrame = -1;
+	}
+
+	if (IsController())
+	{
+		State.PendingEffectCmd.PendingEffect = MoveTemp(Effect);
+	}
+	else
+	{
+		NetworkPredictionProxy.ModifyNetState<FPhysicsEffectsAsyncModelDef>([CapturedEffect = MoveTemp(Effect)](FPhysicsEffectNetState& NetState)
+		{
+			NetState.ActiveEffect = CapturedEffect;
+		});
+	}
+}
+
 void IPhysicsEffectsInterface::ApplyImpulse(float Radius, float Magnitude, FTransform RelativeOffset, float FudgeVelocityZ, const FPhysicsEffectParams& Params)
 {
 	FPhysicsEffectDef NewEffect;
diff --git a/Engine/Plugins/Runtime/NetworkPredictionExtras/Source/NetworkPredictionExtras/Public/PhysicsEffect.h b/Engine/Plugins/Runtime/NetworkPredictionExtras/Source/NetworkPredictionExtras/Public/PhysicsEffect.h
--- a/Engine/Plugins/Runtime/NetworkPredictionExtras/Source/NetworkPredictionExtras/Public/PhysicsEffect.h
+++ b/Engine/Plugins/Runtime/NetworkPredictionExtras/Source/NetworkPredictionExtras/Public/PhysicsEffect.h
@@ -32,8 +32,10 @@ struct FPhysicsEffectsExternalState
 {
 	FPhysicsEffectCmd PendingEffectCmd;
 	FPhysicsEffectNetState PhysicsEffectState;
+	
 	int32 CachedLastPlayedSimFrame = 0;
-	int32 CachedWindUpSimFrame = 0;
+	uint8 CachedActiveTypeID = 0;
+	uint8 CachedWindUpTypeID = 0;
 };
 
 USTRUCT(BlueprintType, meta=(ShowOnlyInnerProperties))
@@ -81,9 +83,15 @@ public:
 	// Notify that a PE has executed (on the PT) and its output/side effects/(whatever it did) has been marshaled back (to the GT)
 	virtual void OnPhysicsEffectExecuted(uint8 TypeID) { };
 
+	// Notify that a PE has executed (on the PT) and its output/side effects/(whatever it did) has been marshaled back (to the GT)
+	virtual void OnPhysicsEffectEnd(uint8 TypeID) { };
+
 	// Notify a PE is pending execute but hasn't been marshaled back to the GT yet
 	virtual void OnPhysicsEffectWindUp(uint8 TypeID) { };
 
+	// Notify a PE windup is "done" which could mean it actually executed (OnPhysicsEffectExecuted) will be called or not (It was cleared or misprediction, etc)
+	virtual void OnPhysicsEffectWindUpEnd(uint8 TypeID) { };
+
 	// Call in InitializeComponent() or anytime after the NP Async Proxy has been registered
 	void InitializePhysicsEffects(FPhysicsEffectLocalState&& LocalState);
 
@@ -118,8 +126,7 @@ public:
 	virtual void ApplyImpulse(float Radius, float Magnitude, FTransform RelativeOffset, float FudgeVelocityZ, const FPhysicsEffectParams& Params);
 
 private:
-
-	void UpdatePendingWindUp(int32 Frame, uint8 TypeID);
+	
 	void ApplyPhysicsEffectDef(FPhysicsEffectDef&& Effect, const FPhysicsEffectParams& Params);
 };
 
diff --git a/Engine/Plugins/Runtime/NetworkPredictionExtras/Source/NetworkPredictionExtras/Public/PhysicsEffectsComponent.h b/Engine/Plugins/Runtime/NetworkPredictionExtras/Source/NetworkPredictionExtras/Public/PhysicsEffectsComponent.h
--- a/Engine/Plugins/Runtime/NetworkPredictionExtras/Source/NetworkPredictionExtras/Public/PhysicsEffectsComponent.h
+++ b/Engine/Plugins/Runtime/NetworkPredictionExtras/Source/NetworkPredictionExtras/Public/PhysicsEffectsComponent.h
@@ -39,11 +39,16 @@ public:
 	FOnPhysicsEffect NotifyPhysicsEffectExecuted;
 	void OnPhysicsEffectExecuted(uint8 TypeID) override { NotifyPhysicsEffectExecuted.Broadcast(TypeID); };
 
-	// Delegate to inform the GT/bp  that a PE will take place soon but not hasn't happened yet
+	// Delegate to inform the GT/bp that a PE will take place soon but not hasn't happened yet
 	UPROPERTY(BlueprintAssignable, Category = "Physics Effect")
 	FOnPhysicsEffect NotifyPhysicsEffectWindUp;
 	void OnPhysicsEffectWindUp(uint8 TypeID) override { NotifyPhysicsEffectWindUp.Broadcast(TypeID); };
 
+	// Delegate to inform the GT/bp that a PE wind up has ended
+	UPROPERTY(BlueprintAssignable, Category = "Physics Effect")
+	FOnPhysicsEffect NotifyPhysicsEffectWindUpEnd;
+	void OnPhysicsEffectWindUpEnd(uint8 TypeID) override { NotifyPhysicsEffectWindUpEnd.Broadcast(TypeID); };
+
 private:
 
 	FPhysicsEffectsExternalState PhysicsEffectsExternalState;
diff --git a/Engine/Plugins/Runtime/NetworkPredictionExtras/Source/NetworkPredictionExtras/Public/PhysicsMovement.h b/Engine/Plugins/Runtime/NetworkPredictionExtras/Source/NetworkPredictionExtras/Public/PhysicsMovement.h
--- a/Engine/Plugins/Runtime/NetworkPredictionExtras/Source/NetworkPredictionExtras/Public/PhysicsMovement.h
+++ b/Engine/Plugins/Runtime/NetworkPredictionExtras/Source/NetworkPredictionExtras/Public/PhysicsMovement.h
@@ -313,6 +313,11 @@ public:
 	FOnPhysicsEffect NotifyPhysicsEffectWindUp;
 	void OnPhysicsEffectWindUp(uint8 TypeID) override { NotifyPhysicsEffectWindUp.Broadcast(TypeID); };
 
+	// Delegate to inform the GT/bp that a PE wind up has ended
+	UPROPERTY(BlueprintAssignable, Category = "Physics Effect")
+	FOnPhysicsEffect NotifyPhysicsEffectWindUpEnd;
+	void OnPhysicsEffectWindUpEnd(uint8 TypeID) override { NotifyPhysicsEffectWindUpEnd.Broadcast(TypeID); };
+
 private:
 
 	FPhysicsEffectsExternalState PhysicsEffectsExternalState;
