Network Physics
Lower movement speed a bit
camera support on gamepad

#rnx

diff --git a/Engine/Plugins/Runtime/NetworkPredictionExtras/Source/NetworkPredictionExtras/Private/MockNetworkPhysics.cpp b/Engine/Plugins/Runtime/NetworkPredictionExtras/Source/NetworkPredictionExtras/Private/MockNetworkPhysics.cpp
--- a/Engine/Plugins/Runtime/NetworkPredictionExtras/Source/NetworkPredictionExtras/Private/MockNetworkPhysics.cpp
+++ b/Engine/Plugins/Runtime/NetworkPredictionExtras/Source/NetworkPredictionExtras/Private/MockNetworkPhysics.cpp
@@ -124,7 +124,7 @@ namespace UE_NETWORK_PHYSICS
 	float DragK=200.f;
 	FAutoConsoleVariableRef CVarDragK(TEXT("np2.Mock.DragK"), DragK, TEXT("Drag Coefficient (higher=more drag)"));
 
-	float MovementK=2.5f;
+	float MovementK=1.5f;
 	FAutoConsoleVariableRef CVarMovementK(TEXT("np2.Mock.MovementK"), MovementK, TEXT("Movement Coefficient (higher=faster movement)"));
 
 	float TurnK=100000.f;
@@ -141,6 +141,30 @@ void FMockManagedState::AsyncTick(UWorld* World, Chaos::FPhysicsSolver* Solver,
 	{
 		if (auto* PT = Proxy->GetPhysicsThreadAPI())
 		{
+			/*
+			FVector TracePosition = PT->X();
+			FVector EndPosition = TracePosition + FVector(0.f, 0.f, -255.f);
+			FCollisionShape Shape = FCollisionShape::MakeSphere(250.f);
+			ECollisionChannel CollisionChannel = ECollisionChannel::ECC_WorldStatic; 
+			FCollisionQueryParams QueryParams = FCollisionQueryParams::DefaultQueryParam;
+			FCollisionResponseParams ResponseParams = FCollisionResponseParams::DefaultResponseParam;
+			FCollisionObjectQueryParams ObjectParams(ECollisionChannel::ECC_PhysicsBody);
+
+			//TArray<FOverlapResult> Overlaps;
+			//World->OverlapMultiByChannel(Overlaps, TracePosition, FQuat::Identity, CollisionChannel, Shape);			
+
+			if (World->LineTraceTestByChannel(TracePosition, EndPosition, ECollisionChannel::ECC_WorldStatic, QueryParams, ResponseParams) == false)
+			{
+				UE_LOG(LogTemp, Warning, TEXT("In Air: %d"), SimulationFrame);
+			}
+			else
+			{
+				
+			}
+			*/
+
+			// ---------------------------------------------------------------------------------------------
+
 			const float UpDot = FVector::DotProduct(PT->R().GetUpVector(), FVector::UpVector);
 			if (PT_State.RecoveryFrame == 0)
 			{
@@ -159,7 +183,7 @@ void FMockManagedState::AsyncTick(UWorld* World, Chaos::FPhysicsSolver* Solver,
 				}
 				else
 				{
-					UE_LOG(LogTemp, Warning, TEXT("In Air"));
+					UE_LOG(LogTemp, Warning, TEXT("Recovering"));
 
 					FRotator Rot = PT->R().Rotator();
 					const float DeltaRoll = FRotator::NormalizeAxis( -1.f * (Rot.Roll + (PT->W().X * UE_NETWORK_PHYSICS::TurnDampK)));
@@ -709,7 +733,7 @@ void FMockObjectManager::PreNetSend(UWorld* World, float DeltaSeconds)
 	
 	FMockAsyncObjectManagerInput* AsyncInput = AsyncCallback->GetProducerInputData_External();
 	AsyncInput->Reset();	//only want latest frame's data
-
+	AsyncInput->World = World;
 	AsyncInput->ManagedObjects.Reserve(InMockManagedStates.Num());
 		
 	for (FMockManagedState* State : InMockManagedStates)
