Network Physics
-Tweaked jump implementation to not be cooldown based and is now analog

#rnx

diff --git a/Engine/Plugins/Runtime/NetworkPredictionExtras/Source/NetworkPredictionExtras/Private/MockNetworkPhysics.cpp b/Engine/Plugins/Runtime/NetworkPredictionExtras/Source/NetworkPredictionExtras/Private/MockNetworkPhysics.cpp
--- a/Engine/Plugins/Runtime/NetworkPredictionExtras/Source/NetworkPredictionExtras/Private/MockNetworkPhysics.cpp
+++ b/Engine/Plugins/Runtime/NetworkPredictionExtras/Source/NetworkPredictionExtras/Private/MockNetworkPhysics.cpp
@@ -132,6 +132,15 @@ namespace UE_NETWORK_PHYSICS
 
 	float TurnDampK=20.f;
 	FAutoConsoleVariableRef CVarTurnDampK(TEXT("np2.Mock.TurnDampK"), TurnDampK, TEXT("Coefficient for damping portion of turn. Higher=more damping but too higher will lead to instability."));
+
+	float JumpForce=800000.0f;
+	FAutoConsoleVariableRef CVarJumpForce(TEXT("np2.Mock.JumpForce"), JumpForce, TEXT("Per-Frame force to apply while jumping."));
+
+	int32 JumpFrameDuration=3;
+	FAutoConsoleVariableRef CVarJumpFrameDuration(TEXT("np2.Mock.JumpFrameDuration"), JumpFrameDuration, TEXT("How many frames to apply jump force for"));
+
+	int32 JumpFudgeFrames=10;
+	FAutoConsoleVariableRef CVarJumpFudgeFrames(TEXT("np2.Mock.JumpFudgeFrames"), JumpFudgeFrames, TEXT("How many frames after being in air do we still allow a jump to begin"));
 }
 
 void FMockManagedState::AsyncTick(UWorld* World, Chaos::FPhysicsSolver* Solver, const float DeltaSeconds, const int32 SimulationFrame, const int32 LocalStorageFrame)
@@ -141,39 +150,74 @@ void FMockManagedState::AsyncTick(UWorld* World, Chaos::FPhysicsSolver* Solver,
 	{
 		if (auto* PT = Proxy->GetPhysicsThreadAPI())
 		{
-			/*
+			
 			FVector TracePosition = PT->X();
-			FVector EndPosition = TracePosition + FVector(0.f, 0.f, -255.f);
+			FVector EndPosition = TracePosition + FVector(0.f, 0.f, -100.f);
 			FCollisionShape Shape = FCollisionShape::MakeSphere(250.f);
 			ECollisionChannel CollisionChannel = ECollisionChannel::ECC_WorldStatic; 
 			FCollisionQueryParams QueryParams = FCollisionQueryParams::DefaultQueryParam;
 			FCollisionResponseParams ResponseParams = FCollisionResponseParams::DefaultResponseParam;
 			FCollisionObjectQueryParams ObjectParams(ECollisionChannel::ECC_PhysicsBody);
 
-			//TArray<FOverlapResult> Overlaps;
-			//World->OverlapMultiByChannel(Overlaps, TracePosition, FQuat::Identity, CollisionChannel, Shape);			
+			FHitResult OutHit;
+			const bool bInAir = !World->LineTraceSingleByChannel(OutHit, TracePosition, EndPosition, ECollisionChannel::ECC_WorldStatic, QueryParams, ResponseParams);
 
-			if (World->LineTraceTestByChannel(TracePosition, EndPosition, ECollisionChannel::ECC_WorldStatic, QueryParams, ResponseParams) == false)
+			const float UpDot = FVector::DotProduct(PT->R().GetUpVector(), FVector::UpVector);
+
+			// ---------------------------------------------------------------------------------------------
+						
+			if (!bInAir)
 			{
-				UE_LOG(LogTemp, Warning, TEXT("In Air: %d"), SimulationFrame);
+				if (PT_State.InAirFrame != 0)
+				{
+					PT_State.InAirFrame = 0;
+					//PT_State.JumpStartFrame = 0;
+				}
+
+				// Check for recovery start
+				if (PT_State.RecoveryFrame == 0)
+				{
+					if (UpDot < 0.2f)
+					{
+						PT_State.RecoveryFrame = SimulationFrame;
+					}
+				}
 			}
 			else
 			{
-				
+				if (PT_State.InAirFrame == 0)
+				{
+					PT_State.InAirFrame = SimulationFrame;
+				}
 			}
-			*/
-
-			// ---------------------------------------------------------------------------------------------
+			
+			//UE_LOG(LogNetworkPhysics, Log, TEXT("[%d] AirFrame: %d. JumpFrame: %d"), SimulationFrame, PT_State.InAirFrame, PT_State.JumpStartFrame);
+			if (InputCmd.bJumpedPressed)
+			{
+				if (PT_State.InAirFrame == 0 || (PT_State.InAirFrame + UE_NETWORK_PHYSICS::JumpFudgeFrames > SimulationFrame))
+				{
+					if (PT_State.JumpStartFrame == 0)
+					{
+						PT_State.JumpStartFrame = SimulationFrame;
+					}
 
-			const float UpDot = FVector::DotProduct(PT->R().GetUpVector(), FVector::UpVector);
-			if (PT_State.RecoveryFrame == 0)
+					if (PT_State.JumpStartFrame + UE_NETWORK_PHYSICS::JumpFrameDuration > SimulationFrame)
+					{
+						PT->AddForce( Chaos::FVec3(0.f, 0.f, UE_NETWORK_PHYSICS::JumpForce) );
+					
+						//UE_LOG(LogTemp, Warning, TEXT("[%d] Jumped [JumpStart: %d. InAir: %d]"), SimulationFrame, PT_State.JumpStartFrame, PT_State.InAirFrame);
+						PT_State.JumpCooldownMS = 1000;
+					}
+				}
+			}
+			else
 			{
-				if (UpDot < 0.2f)
+				if (PT_State.InAirFrame == 0 && (PT_State.JumpStartFrame + UE_NETWORK_PHYSICS::JumpFrameDuration < SimulationFrame))
 				{
-					PT_State.RecoveryFrame = SimulationFrame;
+					PT_State.JumpStartFrame = 0;
 				}
 			}
-			
+
 			if (PT_State.RecoveryFrame != 0)
 			{
 				if (UpDot > 0.7f)
@@ -183,8 +227,7 @@ void FMockManagedState::AsyncTick(UWorld* World, Chaos::FPhysicsSolver* Solver,
 				}
 				else
 				{
-					UE_LOG(LogTemp, Warning, TEXT("Recovering"));
-
+					// Doing it per-axis like this is probably wrong
 					FRotator Rot = PT->R().Rotator();
 					const float DeltaRoll = FRotator::NormalizeAxis( -1.f * (Rot.Roll + (PT->W().X * UE_NETWORK_PHYSICS::TurnDampK)));
 					const float DeltaPitch = FRotator::NormalizeAxis( -1.f * (Rot.Pitch + (PT->W().Y * UE_NETWORK_PHYSICS::TurnDampK)));
@@ -193,7 +236,6 @@ void FMockManagedState::AsyncTick(UWorld* World, Chaos::FPhysicsSolver* Solver,
 					PT->AddForce(FVector(0.f, 0.f, 600.f));
 				}
 			}
-
 			else if (InputCmd.bBrakesPressed)
 			{
 				FVector NewV = PT->V();
@@ -205,75 +247,43 @@ void FMockManagedState::AsyncTick(UWorld* World, Chaos::FPhysicsSolver* Solver,
 				{
 					PT->SetV( Chaos::FVec3(NewV.X * 0.8f, NewV.Y * 0.8f, NewV.Z));
 				}
-				/*
-
-				PT->SetV( Chaos::FVec3(0.f));
-				PT->SetW( Chaos::FVec3(0.f));
-
-				UE_CLOG(UE_NETWORK_PHYSICS::MockDebug, LogNetworkPhysics, Log, TEXT("[%d/%d] Applied Break. Rot was: %s"), SimulationFrame, LocalStorageFrame, *PT->R().Rotator().ToString());
-				
-				Chaos::FRotation3 NewR = PT->R();
-				FRotator Rot = NewR.Rotator();
-				Rot.Pitch = 0.f;
-				Rot.Roll = 0.f;
-				//PT->SetR( FQuat(Rot));				
-				//PT->SetR(FQuat(FRotator(-90.f, -71.f, 39.f)));
-				PT->SetR(FQuat(FRotator(-90.f, -71.f, 39.f)));
-				*/
 			}
 			else
 			{
+				// Movement
 				if (InputCmd.Force.SizeSquared() > 0.001f)
 				{
-					//UE_LOG(LogTemp, Warning, TEXT("Applied Force. Frame: %d"), LocalFrame);
-					//UE_LOG(LogTemp, Warning, TEXT("0x%X ForceMultiplier: %f (Rand: %d)"), (int64)Proxy, GT_State.ForceMultiplier, GT_State.RandValue);
-					//UE_CLOG(PC == nullptr && InputCmd.Force.SizeSquared() > 0.f, LogTemp, Warning, TEXT("0x%X ForceMultiplier: %f (Rand: %d)"), (int64)Proxy, GT_State.ForceMultiplier, GT_State.RandValue);
-					
 					PT->AddForce(InputCmd.Force * GT_State.ForceMultiplier * UE_NETWORK_PHYSICS::MovementK);
 
+					// Auto Turn
 					const float CurrentYaw = PT->R().Rotator().Yaw + (PT->W().Z * UE_NETWORK_PHYSICS::TurnDampK);
 					const float DesiredYaw = InputCmd.Force.Rotation().Yaw;
 					const float DeltaYaw = FRotator::NormalizeAxis( DesiredYaw - CurrentYaw );
 					
 					PT->AddTorque(FVector(0.f, 0.f, DeltaYaw * UE_NETWORK_PHYSICS::TurnK));
 				}
-
-				FVector V = PT->V();
-				V.Z = 0.f;
-				if (V.SizeSquared() > 0.1f)
-				{
-					FVector Drag = -1.f * V * UE_NETWORK_PHYSICS::DragK;
-					PT->AddForce(Drag);
-				}
-
-				
-				if (FMath::Abs<float>(InputCmd.Turn) > 0.001f)
-				{
-					PT->AddTorque(Chaos::FVec3(0.0f, 0.0f, InputCmd.Turn * GT_State.ForceMultiplier * 10.f));
-				}				
 			}
 
-			if (InputCmd.bJumpedPressed)
+			// Drag force
+			FVector V = PT->V();
+			V.Z = 0.f;
+			if (V.SizeSquared() > 0.1f)
 			{
-				PT_State.JumpCount++;
-				UE_CLOG(UE_NETWORK_PHYSICS::MockDebug, LogNetworkPhysics, Log, TEXT("[%d/%d] bJumpedPressed: %d. Count: %d"), SimulationFrame, LocalStorageFrame, InputCmd.bJumpedPressed, PT_State.JumpCount);
+				FVector Drag = -1.f * V * UE_NETWORK_PHYSICS::DragK;
+				PT->AddForce(Drag);
 			}
-
-			//UE_CLOG(PC == nullptr, LogTemp, Warning, TEXT("NP [%d] bJumpedPressed: %d. Count: %d"), LocalFrame, InputCmd.bJumpedPressed, PT_State.JumpCount);
-			if (InputCmd.bJumpedPressed && PT_State.JumpCooldownMS == 0)
+			
+			PT_State.JumpCooldownMS = FMath::Max( PT_State.JumpCooldownMS - (int32)(DeltaSeconds* 1000.f), 0);
+			if (PT_State.JumpCooldownMS != 0)
 			{
-				//UE_LOG(LogTemp, Warning, TEXT("0x%X Jump! Total: %d"), (int64)Proxy, (int32)(TotalSeconds*1000.f));
-				PT->SetLinearImpulse( Chaos::FVec3(0.f, 0.f, 115000.f) );
-				PT_State.JumpCooldownMS = 1000;
-				UE_CLOG(UE_NETWORK_PHYSICS::MockDebug, LogNetworkPhysics, Log, TEXT("[%d/%d] Applied Jump and reset cooldown"), SimulationFrame, LocalStorageFrame);
+				UE_CLOG(UE_NETWORK_PHYSICS::MockDebug, LogNetworkPhysics, Log, TEXT("[%d/%d] JumpCount: %d. JumpCooldown: %d"), SimulationFrame, LocalStorageFrame, PT_State.JumpCount, PT_State.JumpCooldownMS);
 			}
-			else
+
+			if (InputCmd.bJumpedPressed)
 			{
-				PT_State.JumpCooldownMS = FMath::Max( PT_State.JumpCooldownMS - (int32)(DeltaSeconds* 1000.f), 0);
-				if (PT_State.JumpCooldownMS != 0)
-				{
-					UE_CLOG(UE_NETWORK_PHYSICS::MockDebug, LogNetworkPhysics, Log, TEXT("[%d/%d] JumpCount: %d. JumpCooldown: %d"), SimulationFrame, LocalStorageFrame, PT_State.JumpCount, PT_State.JumpCooldownMS);
-				}
+				// Note this is really just for debugging. "How many times was the button pressed"
+				PT_State.JumpCount++;
+				UE_CLOG(UE_NETWORK_PHYSICS::MockDebug, LogNetworkPhysics, Log, TEXT("[%d/%d] bJumpedPressed: %d. Count: %d"), SimulationFrame, LocalStorageFrame, InputCmd.bJumpedPressed, PT_State.JumpCount);
 			}
 		}
 	}
@@ -901,13 +911,16 @@ void UNetworkPhysicsComponent::InitializeComponent()
 		});
 
 		
-		FMockObjectManager* MockManager = FMockObjectManager::Get(World);
-		checkSlow(MockManager);
+		if (bEnableMockGameplay)
+		{
+			FMockObjectManager* MockManager = FMockObjectManager::Get(World);
+			checkSlow(MockManager);
 
-		InManagedState.Proxy = PrimitiveComponent->BodyInstance.ActorHandle;
-		OutManagedState.Proxy = PrimitiveComponent->BodyInstance.ActorHandle;
-		ReplicatedManagedState.Proxy = PrimitiveComponent->BodyInstance.ActorHandle;
-		MockManager->RegisterManagedMockObject(&ReplicatedManagedState, &InManagedState, &OutManagedState);
+			InManagedState.Proxy = PrimitiveComponent->BodyInstance.ActorHandle;
+			OutManagedState.Proxy = PrimitiveComponent->BodyInstance.ActorHandle;
+			ReplicatedManagedState.Proxy = PrimitiveComponent->BodyInstance.ActorHandle;
+			MockManager->RegisterManagedMockObject(&ReplicatedManagedState, &InManagedState, &OutManagedState);
+		}
 	}
 #endif
 }
@@ -927,10 +940,13 @@ void UNetworkPhysicsComponent::EndPlay(const EEndPlayReason::Type EndPlayReason)
 	{
 		if (UNetworkPhysicsManager* Manager = World->GetSubsystem<UNetworkPhysicsManager>())
 		{
-			if (FMockObjectManager* MockManager = FMockObjectManager::Get(World))
+			if (bEnableMockGameplay)
 			{
-				//UE_LOG(LogTemp, Warning, TEXT("   Unregistering MockManager. 0x%X"), (int64)&ReplicatedManagedState);
-				MockManager->UnregisterManagedMockObject(&ReplicatedManagedState, &InManagedState, &OutManagedState);				
+				if (FMockObjectManager* MockManager = FMockObjectManager::Get(World))
+				{
+					//UE_LOG(LogTemp, Warning, TEXT("   Unregistering MockManager. 0x%X"), (int64)&ReplicatedManagedState);
+					MockManager->UnregisterManagedMockObject(&ReplicatedManagedState, &InManagedState, &OutManagedState);				
+				}
 			}
 			Manager->UnregisterPhysicsProxy(&NetworkPhysicsState);
 		}
diff --git a/Engine/Plugins/Runtime/NetworkPredictionExtras/Source/NetworkPredictionExtras/Public/MockNetworkPhysics.h b/Engine/Plugins/Runtime/NetworkPredictionExtras/Source/NetworkPredictionExtras/Public/MockNetworkPhysics.h
--- a/Engine/Plugins/Runtime/NetworkPredictionExtras/Source/NetworkPredictionExtras/Public/MockNetworkPhysics.h
+++ b/Engine/Plugins/Runtime/NetworkPredictionExtras/Source/NetworkPredictionExtras/Public/MockNetworkPhysics.h
@@ -22,9 +22,6 @@ struct FMockPhysInputCmd
 	UPROPERTY(BlueprintReadWrite, Category="Input")
 	FVector	Force;
 
-	UPROPERTY(BlueprintReadWrite, Category="Input")
-	float Turn=0.f;
-
 	UPROPERTY(BlueprintReadWrite, Category="Input")
 	bool bJumpedPressed = false;
 
@@ -34,7 +31,6 @@ struct FMockPhysInputCmd
 	bool NetSerialize(FArchive& Ar, class UPackageMap* Map, bool& bOutSuccess)
 	{
 		Ar << Force;
-		Ar << Turn;
 		Ar << bJumpedPressed;
 		Ar << bBrakesPressed;
 		return true;
@@ -52,23 +48,18 @@ struct FMockPhysInputCmd
 	bool ShouldReconcile(const FMockPhysInputCmd& AuthState) const
 	{
 		return FVector::DistSquared(Force, AuthState.Force) > 0.1f || bJumpedPressed != AuthState.bJumpedPressed || 
-			bBrakesPressed != AuthState.bBrakesPressed || FMath::Abs<float>(Turn - AuthState.Turn) > 0.1f;
+			bBrakesPressed != AuthState.bBrakesPressed;
 	}
 
 	// Decays input for non locally controlled sims (client only)
 	// disabled behind np2.InputDecay cvar for now
 	void Decay(float Alpha)
 	{
-		Turn *= Alpha;
-		if (Turn < 0.001f)
-		{
-			Turn = 0.f;
-		}
 	}
 
 	void ToString(FStringBuilderBase& Out)
 	{
-		Out.Appendf(TEXT("%s. bJumpedPressed: %d. bBrakesPressed: %d. Turn: %.2f"), *Force.ToString(), (int32)bJumpedPressed, (int32)bBrakesPressed, Turn);
+		Out.Appendf(TEXT("%s. bJumpedPressed: %d. bBrakesPressed: %d."), *Force.ToString(), (int32)bJumpedPressed, (int32)bBrakesPressed);
 	}
 };
 
@@ -146,18 +137,28 @@ struct FMockState_PT
 	UPROPERTY(BlueprintReadWrite, Category="Mock Object")
 	int32 RecoveryFrame = 0;
 
+	// Frame we started jumping on
+	UPROPERTY(BlueprintReadWrite, Category="Mock Object")
+	int32 JumpStartFrame = 0;
+
+	// Frame we started being in the air
+	UPROPERTY(BlueprintReadWrite, Category="Mock Object")
+	int32 InAirFrame = 0;
+
 	bool NetSerialize(FArchive& Ar, class UPackageMap* Map, bool& bOutSuccess)
 	{
 		Ar << JumpCooldownMS;
 		Ar << JumpCount;
 		Ar << CheckSum;	
 		Ar << RecoveryFrame;
+		Ar << JumpStartFrame;
+		Ar << InAirFrame;
 		return true;
 	}
 
 	bool ShouldReconcile(const FMockState_PT& AuthState) const
 	{
-		return JumpCooldownMS != AuthState.JumpCooldownMS || JumpCount != AuthState.JumpCount || RecoveryFrame != AuthState.RecoveryFrame;
+		return JumpCooldownMS != AuthState.JumpCooldownMS || JumpCount != AuthState.JumpCount || RecoveryFrame != AuthState.RecoveryFrame || JumpStartFrame != AuthState.JumpStartFrame || InAirFrame != AuthState.InAirFrame;
 	}
 };
 
@@ -330,6 +331,10 @@ public:
 	UFUNCTION(BlueprintPure, Category = "Network Physics")
 	int32 GetNetworkPredictionLOD() const { return NetworkPhysicsState.LocalLOD; }
 
+	// Makes the object "controllable". This includes logic like air-recover.
+	UPROPERTY(EditDefaultsOnly, Category="Network Physics")
+	bool bEnableMockGameplay=false;
+
 protected:
 
 	// Managed state should not be publically exposed
