Network Physics
-Tweak jump some more
-Remove some log spam

#rnx

diff --git a/Engine/Plugins/Runtime/NetworkPrediction/Source/NetworkPrediction/Private/NetworkPhysics.cpp b/Engine/Plugins/Runtime/NetworkPrediction/Source/NetworkPrediction/Private/NetworkPhysics.cpp
--- a/Engine/Plugins/Runtime/NetworkPrediction/Source/NetworkPrediction/Private/NetworkPhysics.cpp
+++ b/Engine/Plugins/Runtime/NetworkPrediction/Source/NetworkPrediction/Private/NetworkPhysics.cpp
@@ -149,6 +149,7 @@ struct FNetworkPhysicsRewindCallback : public Chaos::IRewindCallback
 				if (CompareObjState(Obj.Physics.ObjectState, P.ObjectState(), TEXT("DEBUG ObjectState")))
 				//if (CompareVec(Obj.Physics.Location, P.X(), UE_NETWORK_PHYSICS::X, TEXT("Location")))
 				{
+					/*
 					UE_LOG(LogTemp, Warning, TEXT("ObjectState Correction. Obj.Frame: %d. CurrentFramE: %d"), Obj.Frame, RewindData->CurrentFrame());
 					//UE_LOG(LogTemp, Warning, TEXT("Location Correction. Obj.Frame: %d. CurrentFrame: %d"), Obj.Frame, RewindData->CurrentFrame());
 					
@@ -160,6 +161,7 @@ struct FNetworkPhysicsRewindCallback : public Chaos::IRewindCallback
 						UE_LOG(LogTemp, Warning, TEXT("    [%d/%d] --> %d"), F, F-Snapshot.LocalFrameOffset, (int32)PP.ObjectState());
 					}
 					UE_LOG(LogTemp, Warning, TEXT(""));
+					*/
 					
 				}
 
diff --git a/Engine/Plugins/Runtime/NetworkPredictionExtras/Source/NetworkPredictionExtras/Private/MockNetworkPhysics.cpp b/Engine/Plugins/Runtime/NetworkPredictionExtras/Source/NetworkPredictionExtras/Private/MockNetworkPhysics.cpp
--- a/Engine/Plugins/Runtime/NetworkPredictionExtras/Source/NetworkPredictionExtras/Private/MockNetworkPhysics.cpp
+++ b/Engine/Plugins/Runtime/NetworkPredictionExtras/Source/NetworkPredictionExtras/Private/MockNetworkPhysics.cpp
@@ -124,7 +124,7 @@ namespace UE_NETWORK_PHYSICS
 	float DragK=200.f;
 	FAutoConsoleVariableRef CVarDragK(TEXT("np2.Mock.DragK"), DragK, TEXT("Drag Coefficient (higher=more drag)"));
 
-	float MovementK=1.5f;
+	float MovementK=1.25f;
 	FAutoConsoleVariableRef CVarMovementK(TEXT("np2.Mock.MovementK"), MovementK, TEXT("Movement Coefficient (higher=faster movement)"));
 
 	float TurnK=100000.f;
@@ -133,10 +133,10 @@ namespace UE_NETWORK_PHYSICS
 	float TurnDampK=20.f;
 	FAutoConsoleVariableRef CVarTurnDampK(TEXT("np2.Mock.TurnDampK"), TurnDampK, TEXT("Coefficient for damping portion of turn. Higher=more damping but too higher will lead to instability."));
 
-	float JumpForce=800000.0f;
+	float JumpForce=1000000.0f;
 	FAutoConsoleVariableRef CVarJumpForce(TEXT("np2.Mock.JumpForce"), JumpForce, TEXT("Per-Frame force to apply while jumping."));
 
-	int32 JumpFrameDuration=3;
+	int32 JumpFrameDuration=4;
 	FAutoConsoleVariableRef CVarJumpFrameDuration(TEXT("np2.Mock.JumpFrameDuration"), JumpFrameDuration, TEXT("How many frames to apply jump force for"));
 
 	int32 JumpFudgeFrames=10;
