Network Physics input recording
-Allows recording input streams and playing them back for debugging/testing

#rnx

diff --git a/Engine/Plugins/Runtime/NetworkPrediction/Source/NetworkPrediction/Private/NetworkPredictionDebug.cpp b/Engine/Plugins/Runtime/NetworkPrediction/Source/NetworkPrediction/Private/NetworkPredictionDebug.cpp
--- a/Engine/Plugins/Runtime/NetworkPrediction/Source/NetworkPrediction/Private/NetworkPredictionDebug.cpp
+++ b/Engine/Plugins/Runtime/NetworkPrediction/Source/NetworkPrediction/Private/NetworkPredictionDebug.cpp
@@ -6,6 +6,11 @@
 #include "GameFramework/Actor.h"
 #include "DrawDebugHelpers.h"
 #include "NetworkPredictionWorldManager.h"
+#include "Misc/NetworkGuid.h"
+#include "Engine/NetConnection.h"
+#include "Engine/NetDriver.h"
+#include "EngineUtils.h"
+#include "UObject/UObjectIterator.h"
 
 namespace NetworkPredictionDebug
 {
@@ -31,6 +36,49 @@ namespace NetworkPredictionDebug
 
 		DrawDebugString(World, Transform.GetLocation(), Str, nullptr, DrawColor, Lifetime, false);
 	}
+
+	UObject* FindReplicatedObjectOnPIEServer(UObject* ClientObject)
+	{
+		if (ClientObject == nullptr)
+		{
+			return nullptr;
+		}
+
+		UObject* ServerObject = nullptr;
+
+#if WITH_EDITOR
+		if (UWorld* World = ClientObject->GetWorld())
+		{
+			if (UNetDriver* NetDriver = World->GetNetDriver())
+			{
+				if (UNetConnection* NetConnection = NetDriver->ServerConnection)
+				{
+					FNetworkGUID NetGUID = NetConnection->PackageMap->GetNetGUIDFromObject(ClientObject);
+
+					// Find the PIE server world
+					for (TObjectIterator<UWorld> It; It; ++It)
+					{
+						if (It->WorldType == EWorldType::PIE && It->GetNetMode() != NM_Client)
+						{
+							if (UNetDriver* ServerNetDriver = It->GetNetDriver())
+							{
+								if (ServerNetDriver->ClientConnections.Num() > 0)
+								{
+									UPackageMap* ServerPackageMap = ServerNetDriver->ClientConnections[0]->PackageMap;
+									ServerObject = ServerPackageMap->GetObjectFromNetGUID(NetGUID, true);
+									break;
+								}
+							}
+						}
+					}
+				}
+			}
+		}
+#endif
+
+		return ServerObject;
+	}
+
 };
 
 
diff --git a/Engine/Plugins/Runtime/NetworkPrediction/Source/NetworkPrediction/Public/NetworkPredictionDebug.h b/Engine/Plugins/Runtime/NetworkPrediction/Source/NetworkPrediction/Public/NetworkPredictionDebug.h
--- a/Engine/Plugins/Runtime/NetworkPrediction/Source/NetworkPrediction/Public/NetworkPredictionDebug.h
+++ b/Engine/Plugins/Runtime/NetworkPrediction/Source/NetworkPrediction/Public/NetworkPredictionDebug.h
@@ -10,4 +10,5 @@ namespace NetworkPredictionDebug
 {
 	NETWORKPREDICTION_API void DrawDebugOutline(FTransform Transform, FBox BoundingBox, FColor Color, float Lifetime);
 	NETWORKPREDICTION_API void DrawDebugText3D(const TCHAR* Str, FTransform Transform, FColor, float Lifetime);
+	NETWORKPREDICTION_API UObject* FindReplicatedObjectOnPIEServer(UObject* ClientObject);
 };
\ No newline at end of file
diff --git a/Engine/Plugins/Runtime/NetworkPredictionExtras/Source/NetworkPredictionExtras/Private/MockNetworkPhysics.cpp b/Engine/Plugins/Runtime/NetworkPredictionExtras/Source/NetworkPredictionExtras/Private/MockNetworkPhysics.cpp
--- a/Engine/Plugins/Runtime/NetworkPredictionExtras/Source/NetworkPredictionExtras/Private/MockNetworkPhysics.cpp
+++ b/Engine/Plugins/Runtime/NetworkPredictionExtras/Source/NetworkPredictionExtras/Private/MockNetworkPhysics.cpp
@@ -99,6 +99,8 @@ FMockState_PT
 #include "GameFramework/PlayerController.h"
 #include "DrawDebugHelpers.h"
 #include "PhysicsProxy/SingleParticlePhysicsProxy.h"
+#include "TimerManager.h"
+#include "NetworkPredictionDebug.h"
 
 // ==================================================
 
@@ -616,6 +618,7 @@ FMockObjectManager* FMockObjectManager::Get(UWorld* World)
 
 FMockObjectManager::FMockObjectManager(UWorld* World)
 {
+	WeakWorld = World;
 	if (ensure(World))
 	{
 		FPhysScene* PhysScene = World->GetPhysicsScene();
@@ -736,31 +739,7 @@ void FMockObjectManager::PreNetSend(UWorld* World, float DeltaSeconds)
 		return;
 	}
 
-	// -------------------------------------------
-	// Marshall the managed objects to PT
-	// This part needs to happen once per PT tick - currently being called in PreNetSend (probably) every frame on GT is not ideal
-	// -------------------------------------------
 	
-	FMockAsyncObjectManagerInput* AsyncInput = AsyncCallback->GetProducerInputData_External();
-	AsyncInput->Reset();	//only want latest frame's data
-	AsyncInput->World = World;
-	AsyncInput->ManagedObjects.Reserve(InMockManagedStates.Num());
-		
-	for (FMockManagedState* State : InMockManagedStates)
-	{
-		// Should we decay the GT InputCmd or just the marshalled copy?
-		// Probably want non linear decay?
-		if (UE_NETWORK_PHYSICS::bInputDecay && State->PC == nullptr)
-		{
-			if (State->InputDecay > 0.f)
-			{
-				State->InputCmd.Decay(State->InputDecay);
-			}
-			State->InputDecay += DeltaSeconds * UE_NETWORK_PHYSICS::InputDecayRate;
-		}
-
-		AsyncInput->ManagedObjects.Emplace( *State );
-	}
 
 	// -------------------------------------------
 	//	Marshall data from PT
@@ -844,6 +823,42 @@ void FMockObjectManager::PreNetSend(UWorld* World, float DeltaSeconds)
 	}
 }
 
+void FMockObjectManager::ProcessInputs_External(int32 PhysicsStep, int32 LocalFrameOffset, bool& bOutSendClientInputCmd)
+{
+	// -------------------------------------------
+	// Marshall the managed objects to PT
+	// This part needs to happen once per PT tick - currently being called in PreNetSend (probably) every frame on GT is not ideal
+	// -------------------------------------------
+	
+	FMockAsyncObjectManagerInput* AsyncInput = AsyncCallback->GetProducerInputData_External();
+	AsyncInput->Reset();	//only want latest frame's data
+	AsyncInput->World = WeakWorld;
+	AsyncInput->ManagedObjects.Reserve(InMockManagedStates.Num());
+		
+	for (FMockManagedState* State : InMockManagedStates)
+	{
+		// Should we decay the GT InputCmd or just the marshalled copy?
+		// Probably want non linear decay?
+		if (UE_NETWORK_PHYSICS::bInputDecay && State->PC == nullptr)
+		{
+			if (State->InputDecay > 0.f)
+			{
+				State->InputCmd.Decay(State->InputDecay);
+			}
+			//State->InputDecay += DeltaSeconds * UE_NETWORK_PHYSICS::InputDecayRate; FIXME
+		}
+
+		// ------------------------------------------------
+		if (State->Component)
+		{			
+			State->Component->ProcessInputs_External(*State, PhysicsStep, LocalFrameOffset);
+		}
+		// -----------------------------------------------
+
+		AsyncInput->ManagedObjects.Emplace( *State );
+	}
+}
+
 // ========================================================================================
 
 UNetworkPhysicsComponent::UNetworkPhysicsComponent()
@@ -969,7 +984,12 @@ void UNetworkPhysicsComponent::TickComponent(float DeltaTime, enum ELevelTick Ti
 	{
 		// Broadcast out a delegate. The user will use GetPendingInputCmd/SetPendingInputCmd to write to ManagedState.InputCmd
 		OnGeneratedLocalInputCmd.Broadcast();
+		if (bRecording)
+		{
+
+		}
 	}
+	InManagedState.Component = this;
 
 	ReplicatedManagedState.PC = PC;
 	InManagedState.PC = PC;
@@ -994,10 +1014,252 @@ void UNetworkPhysicsComponent::GetLifetimeReplicatedProps(TArray<FLifetimeProper
 	DOREPLIFETIME( UNetworkPhysicsComponent, ReplicatedManagedState);
 }
 
+void UNetworkPhysicsComponent::StartRecording(TArray<FMockPhysInputCmd>* Stream)
+{
+	if (bRecording)
+	{
+		return;
+	}
+
+	CurrentInputCmdStream = Stream;
+	bRecording = true;
+}
+
+void UNetworkPhysicsComponent::StopRecording()
+{
+	if (CurrentInputCmdStream)
+	{
+		UE_LOG(LogTemp, Log, TEXT("Recorded %d Inputs."), CurrentInputCmdStream->Num());
+	}
+
+	bRecording = false;
+	CurrentInputCmdStream = nullptr;
+}
+
+void UNetworkPhysicsComponent::StartPlayback(TArray<FMockPhysInputCmd>* Stream)
+{
+	CurrentInputCmdStream = Stream;
+	if (CurrentInputCmdStream && CurrentInputCmdStream->Num() > 0)
+	{
+		PlaybackIdx = 0;
+	}
+	else
+	{
+		PlaybackIdx = INDEX_NONE;
+	}
+}
+
+void UNetworkPhysicsComponent::ProcessInputs_External(FMockManagedState& State, int32 PhysicsStep, int32 LocalFrameOffset)
+{
+	if (bRecording)
+	{
+		if (CurrentInputCmdStream)
+		{
+			CurrentInputCmdStream->Add(State.InputCmd);
+		}
+	}
+	else if (CurrentInputCmdStream && CurrentInputCmdStream->IsValidIndex(PlaybackIdx))
+	{
+		State.InputCmd = (*CurrentInputCmdStream)[PlaybackIdx++ % CurrentInputCmdStream->Num()];		
+	}
+}
+
+// ============================================================================================================
+
+void ANetworkPredictionSpawner::Spawn(FName StreamName)
+{
+	ANetworkPredictionSpawner* SourceSpawner = nullptr;
+	for (TObjectIterator<UWorld> WorldIt; WorldIt; ++WorldIt)
+	{
+		// Saved inputs are saved on the map version, so we need to copy them over to the server version
+		if (WorldIt->WorldType == EWorldType::Editor)
+		{
+			for (TActorIterator<ANetworkPredictionSpawner> It(*WorldIt); It; ++It)
+			{
+				ANetworkPredictionSpawner* Spawner = *It;
+				if (Spawner->GetName() == this->GetName())
+				{
+					SourceSpawner = Spawner;
+					break;
+				}
+			}
+		}
+	}
+
+	ensure(SourceSpawner);
+
+	StreamName = (StreamName == NAME_None ? SourceSpawner->RecordedInputs[FMath::Rand() % SourceSpawner->RecordedInputs.Num()].Name : StreamName);
+
+	if (FMockRecordedInputs* PlaybackRecordedInputs = SourceSpawner->RecordedInputs.FindByKey(StreamName))
+	{
+		AActor* SpawnedActor = this->GetWorld()->SpawnActor<AActor>(this->SpawnClass.Get(), this->GetActorTransform());
+		if (UNetworkPhysicsComponent* Comp = SpawnedActor->FindComponentByClass<UNetworkPhysicsComponent>())
+		{
+			Comp->StartPlayback(&PlaybackRecordedInputs->Inputs);
+		}
+	}
+	else
+	{
+		UE_LOG(LogNetworkPhysics, Warning, TEXT("Could not find Inputs named %s on %s"), *StreamName.ToString(), *GetName());
+	}
+}
+
+void ANetworkPredictionSpawner::SpawnRandom()
+{
+	Spawn(NAME_None);
+}
+
+void ANetworkPredictionSpawner::StartRecording(UNetworkPhysicsComponent* Target, FName StreamName)
+{
+	FMockRecordedInputs* PlaybackRecordedInputs = RecordedInputs.FindByKey(StreamName);
+	if (!PlaybackRecordedInputs)
+	{
+		PlaybackRecordedInputs = &RecordedInputs.AddDefaulted_GetRef();
+		PlaybackRecordedInputs->Name = StreamName;
+	}
+	PlaybackRecordedInputs->Inputs.Reset();
+
+	// Teleport the target on the server
+	if (AActor* ServerOwner = Cast<AActor>(NetworkPredictionDebug::FindReplicatedObjectOnPIEServer(Target->GetOwner())))
+	{
+		ServerOwner->TeleportTo(this->GetActorLocation(), this->K2_GetActorRotation());
+	}
+
+	// Wait a second and start recording on the client
+	FTimerHandle Handle;
+	Target->GetWorld()->GetTimerManager().SetTimer(Handle, FTimerDelegate::CreateLambda([this, PlaybackRecordedInputs, Target]()
+	{
+		Target->StartRecording(&PlaybackRecordedInputs->Inputs);
+
+	}), 1.f, false);
+}
+
 // ============================================================================================================
 // ============================================================================================================
 // ============================================================================================================
 
+FAutoConsoleCommandWithWorldAndArgs RecordInputCmds(TEXT("np2.RecordInputs"), TEXT(""),
+	FConsoleCommandWithWorldAndArgsDelegate::CreateLambda([](const TArray< FString >& Args, UWorld* InWorld) 
+{
+	FName StreamName = Args.Num() > 0 ? FName(Args[0]) : NAME_None;
+
+	UNetworkPhysicsComponent* RecordTarget = nullptr;
+	for (TObjectIterator<UNetworkPhysicsComponent> It; It; ++It)
+	{
+		UNetworkPhysicsComponent* Comp = *It;
+		if (Comp->GetWorld() == InWorld && Comp->GetOwnerPC() && Comp->GetOwnerPC()->IsLocalPlayerController())
+		{
+			RecordTarget = Comp;
+			break;
+		}
+	}
+
+	if (!RecordTarget)
+	{
+		UE_LOG(LogNetworkPhysics, Warning, TEXT("Could not find viable target to record"));
+		return;
+	}
+
+	if (RecordTarget->IsRecording())
+	{
+		UE_LOG(LogNetworkPhysics, Log, TEXT("Stopped Recording on %s"), *RecordTarget->GetPathName());
+		RecordTarget->StopRecording();
+		return;
+	}
+
+	if (Args.Num() > 1)
+	{
+		// Spawner path
+		for (TObjectIterator<UWorld> WorldIt; WorldIt; ++WorldIt)
+		{
+			if (WorldIt->WorldType == EWorldType::Editor)
+			{
+				for (TActorIterator<ANetworkPredictionSpawner> It(*WorldIt); It; ++It)
+				{
+					ANetworkPredictionSpawner* Spawner = *It;
+					if (Spawner->GetName().Contains(Args[1]))
+					{
+						UE_LOG(LogNetworkPhysics, Warning, TEXT("Recording Stream %s on Spawner %s"), *StreamName.ToString(), *GetPathNameSafe(Spawner));
+						Spawner->StartRecording(RecordTarget, StreamName);
+						break;
+					}
+				}
+				break;
+			}
+		}
+	}
+	else
+	{
+		TArray<FMockPhysInputCmd>* CurrentInputCmdStream = nullptr;
+
+		UNetworkPhysicsComponent* CDO = UNetworkPhysicsComponent::StaticClass()->GetDefaultObject<UNetworkPhysicsComponent>();
+		if (FMockRecordedInputs* MockRecordedInputs = CDO->RecordedInputs.FindByKey(StreamName))
+		{
+			MockRecordedInputs->Inputs.Reset();
+			CurrentInputCmdStream = &MockRecordedInputs->Inputs;
+		}
+		else
+		{
+			FMockRecordedInputs& NewRecordedInputs = CDO->RecordedInputs.AddDefaulted_GetRef();
+			NewRecordedInputs.Name = StreamName;
+			CurrentInputCmdStream = &NewRecordedInputs.Inputs;
+		}
+
+		if (CurrentInputCmdStream)
+		{
+			UE_LOG(LogNetworkPhysics, Log, TEXT("Started %s input recording on %s"), *StreamName.ToString(), *RecordTarget->GetPathName());
+			RecordTarget->StartRecording(CurrentInputCmdStream);
+		}
+	}
+	
+}));
+
+FAutoConsoleCommandWithWorldAndArgs PlaybackInputCmds(TEXT("np2.PlaybackInputs"), TEXT(""),
+	FConsoleCommandWithWorldAndArgsDelegate::CreateLambda([](const TArray< FString >& Args, UWorld* InWorld) 
+{
+	FName StreamName = Args.Num() > 0 ? FName(Args[0]) : NAME_None;
+
+	if (Args.Num() > 1)
+	{		
+		for (TObjectIterator<UWorld> WorldIt; WorldIt; ++WorldIt)
+		{
+			if ((WorldIt->WorldType == EWorldType::Game || WorldIt->WorldType == EWorldType::PIE) && WorldIt->GetNetMode() != NM_Client)
+			{
+				for (TActorIterator<ANetworkPredictionSpawner> It(*WorldIt); It; ++It)
+				{
+					ANetworkPredictionSpawner* Spawner = *It;
+					if (Spawner->GetName().Contains(Args[1]))
+					{
+						UE_LOG(LogNetworkPhysics, Warning, TEXT("Spawning New Stream %s on Spawner %s"), *StreamName.ToString(), *GetPathNameSafe(Spawner));
+						Spawner->Spawn(StreamName);
+						break;
+					}
+				}
+			}
+		}
+	}
+	else
+	{
+		UNetworkPhysicsComponent* CDO = UNetworkPhysicsComponent::StaticClass()->GetDefaultObject<UNetworkPhysicsComponent>();
+		if (FMockRecordedInputs* PlaybackRecordedInputs = CDO->RecordedInputs.FindByKey(StreamName))
+		{
+			if (PlaybackRecordedInputs->Inputs.Num() > 0)
+			{
+				for (TObjectIterator<UNetworkPhysicsComponent> It; It; ++It)
+				{
+					UNetworkPhysicsComponent* Comp = *It;
+					if (Comp->GetWorld() == InWorld && Comp->GetOwnerPC() && Comp->GetOwnerPC()->IsLocalPlayerController())
+					{
+						Comp->StartPlayback(&PlaybackRecordedInputs->Inputs);
+						UE_LOG(LogNetworkPhysics, Log, TEXT("PlayingBack %d Inputs from stream %s."), PlaybackRecordedInputs->Inputs.Num(), *StreamName.ToString());
+						return;
+					}
+				}
+			}
+		}
+	}
+
+}));
 
 FAutoConsoleCommandWithWorldAndArgs ForceMockCorrectionCmd(TEXT("np2.ForceMockCorrection"), TEXT(""),
 	FConsoleCommandWithWorldAndArgsDelegate::CreateLambda([](const TArray< FString >& Args, UWorld* InWorld) 
diff --git a/Engine/Plugins/Runtime/NetworkPredictionExtras/Source/NetworkPredictionExtras/Public/MockNetworkPhysics.h b/Engine/Plugins/Runtime/NetworkPredictionExtras/Source/NetworkPredictionExtras/Public/MockNetworkPhysics.h
--- a/Engine/Plugins/Runtime/NetworkPredictionExtras/Source/NetworkPredictionExtras/Public/MockNetworkPhysics.h
+++ b/Engine/Plugins/Runtime/NetworkPredictionExtras/Source/NetworkPredictionExtras/Public/MockNetworkPhysics.h
@@ -230,6 +230,9 @@ struct FMockManagedState
 		bOutSuccess = true;
 		return true;
 	}
+
+	// This is only used for input recording/playback hack
+	class UNetworkPhysicsComponent* Component = nullptr;
 };
 
 template<>
@@ -261,6 +264,7 @@ public:
 
 	void PostNetRecv(UWorld* World, int32 FrameOffset, int32 LastProcessedFrame) override;
 	void PreNetSend(UWorld* World, float DeltaSeconds) override;
+	void ProcessInputs_External(int32 PhysicsStep, int32 LocalFrameOffset, bool& bOutSendClientInputCmd) override;
 
 private:
 
@@ -268,12 +272,28 @@ private:
 	TArray<FMockManagedState*> InMockManagedStates;
 	TArray<FMockManagedState*> OutMockManagedStates;
 	class FMockAsyncObjectManagerCallback* AsyncCallback = nullptr;
+	
+	TWeakObjectPtr<UWorld> WeakWorld;
 };
 
 // -----------------------------------------------------
 
+USTRUCT(BlueprintType)
+struct FMockRecordedInputs
+{
+	GENERATED_BODY()
+
+	UPROPERTY(EditAnywhere, BlueprintReadWrite, Category="Recorded Inputs")
+	FName Name;
 
+	UPROPERTY(EditAnywhere, BlueprintReadWrite, Category="Recorded Inputs")
+	TArray<FMockPhysInputCmd> Inputs;
 
+	bool operator==(const FName& N) const
+	{
+		return Name == N;
+	}
+};
 
 UCLASS(BlueprintType, meta=(BlueprintSpawnableComponent))
 class NETWORKPREDICTIONEXTRAS_API UNetworkPhysicsComponent : public UActorComponent
@@ -335,6 +355,19 @@ public:
 	UPROPERTY(EditDefaultsOnly, Category="Network Physics")
 	bool bEnableMockGameplay=false;
 
+
+	void ProcessInputs_External(FMockManagedState& State, int32 PhysicsStep, int32 LocalFrameOffset);
+	void StartRecording(TArray<FMockPhysInputCmd>* Stream);
+	void StopRecording();
+	bool IsRecording() const { return bRecording; }
+	void StartPlayback(TArray<FMockPhysInputCmd>* Stream);
+
+	APlayerController* GetOwnerPC() const;
+
+	// Stores all recorded streams on CDO
+	UPROPERTY(EditAnywhere, BlueprintReadWrite, Category="Recorded Inputs")
+	TArray<FMockRecordedInputs>	RecordedInputs;
+
 protected:
 
 	// Managed state should not be publically exposed
@@ -347,6 +380,34 @@ protected:
 	UPROPERTY(transient)
 	FMockManagedState OutManagedState;
 
-	APlayerController* GetOwnerPC() const;
+	// Which one we are playingback or recording to
+	TArray<FMockPhysInputCmd>* CurrentInputCmdStream = nullptr;	
+
+	bool bRecording = false;
+	int32 PlaybackIdx = INDEX_NONE;
+};
+
+UCLASS(BlueprintType, meta=(BlueprintSpawnableComponent))
+class NETWORKPREDICTIONEXTRAS_API ANetworkPredictionSpawner : public AActor
+{
+	GENERATED_BODY()
+public:
+
+	ANetworkPredictionSpawner() = default;
+
+	// Spawns actor and plays named recorded inputs
+	UFUNCTION(BlueprintCallable, Category="Spawner")
+	void Spawn(FName NamedInputs);
+
+	// Spawns actor and plays random stream of recorded input
+	UFUNCTION(BlueprintCallable, Category="Spawner")
+	void SpawnRandom();
+	
+	void StartRecording(UNetworkPhysicsComponent* Target, FName NamedInputs);
+
+	UPROPERTY(EditAnywhere, BlueprintReadWrite, Category="Spawner")
+	TSubclassOf<AActor> SpawnClass;
 
+	UPROPERTY(EditAnywhere, BlueprintReadWrite, Category="Spawner")
+	TArray<FMockRecordedInputs>	RecordedInputs;
 };
\ No newline at end of file
