cis and sa fixes

#rnx

diff --git a/Engine/Plugins/Runtime/NetworkPrediction/Source/NetworkPrediction/Public/NetworkPhysics.h b/Engine/Plugins/Runtime/NetworkPrediction/Source/NetworkPrediction/Public/NetworkPhysics.h
--- a/Engine/Plugins/Runtime/NetworkPrediction/Source/NetworkPrediction/Public/NetworkPhysics.h
+++ b/Engine/Plugins/Runtime/NetworkPrediction/Source/NetworkPrediction/Public/NetworkPhysics.h
@@ -7,7 +7,7 @@
 #include "Chaos/Particles.h"
 #include "Components/ActorComponent.h"
 #include "UObject/ObjectKey.h"
-#include "Misc/NetworkGUID.h"
+#include "Misc/NetworkGuid.h"
 
 #include "NetworkPhysics.generated.h"
 
diff --git a/Engine/Plugins/Runtime/NetworkPredictionExtras/Source/NetworkPredictionExtras/Private/MockNetworkPhysics.cpp b/Engine/Plugins/Runtime/NetworkPredictionExtras/Source/NetworkPredictionExtras/Private/MockNetworkPhysics.cpp
--- a/Engine/Plugins/Runtime/NetworkPredictionExtras/Source/NetworkPredictionExtras/Private/MockNetworkPhysics.cpp
+++ b/Engine/Plugins/Runtime/NetworkPredictionExtras/Source/NetworkPredictionExtras/Private/MockNetworkPhysics.cpp
@@ -147,6 +147,9 @@ namespace UE_NETWORK_PHYSICS
 
 	bool JumpHack=false;
 	FAutoConsoleVariableRef CVarJumpHack(TEXT("np2.Mock.JumpHack"), JumpHack, TEXT("Make jump not rely on trace which currently causes non determinism"));
+
+	float MockImpulse=false;
+	FAutoConsoleVariableRef CVarMockImpulse(TEXT("np2.Mock.BallImpulse"), JumpHack, TEXT("Make jump not rely on trace which currently causes non determinism"));
 }
 
 void FMockManagedState::AsyncTick(UWorld* World, Chaos::FPhysicsSolver* Solver, const float DeltaSeconds, const int32 SimulationFrame, const int32 LocalStorageFrame)
@@ -169,6 +172,9 @@ void FMockManagedState::AsyncTick(UWorld* World, Chaos::FPhysicsSolver* Solver,
 			const bool bInAir = !UE_NETWORK_PHYSICS::JumpHack && !World->LineTraceSingleByChannel(OutHit, TracePosition, EndPosition, ECollisionChannel::ECC_WorldStatic, QueryParams, ResponseParams);
 			const float UpDot = FVector::DotProduct(PT->R().GetUpVector(), FVector::UpVector);
 
+			//TArray<FOverlapResult> OutOverlaps;
+			//FGenericPhysicsInterface::GeomOverlapMulti(World, Shape, TracePosition, FQuat::Identity, OutOverlaps, ECollisionChannel::ECC_PhysicsBody, QueryParams, ResponseParams, ObjectParams);
+
 			// ---------------------------------------------------------------------------------------------
 						
 			if (!bInAir)
@@ -1090,7 +1096,11 @@ AActor* ANetworkPredictionSpawner::Spawn(FName StreamName)
 	}
 
 	AActor* SpawnedActor = nullptr;
-	ensure(SourceSpawner);
+	if (!SourceSpawner)
+	{
+		UE_LOG(LogNetworkPhysics, Warning, TEXT("Could not find spawner named %s"), *StreamName.ToString());
+		return nullptr;
+	}
 
 	StreamName = (StreamName == NAME_None ? SourceSpawner->RecordedInputs[FMath::Rand() % SourceSpawner->RecordedInputs.Num()].Name : StreamName);
 
