Network Physics
-Add kick impulse to soccer game!

#rnx

diff --git a/Engine/Plugins/Runtime/NetworkPredictionExtras/Source/NetworkPredictionExtras/Private/MockNetworkPhysics.cpp b/Engine/Plugins/Runtime/NetworkPredictionExtras/Source/NetworkPredictionExtras/Private/MockNetworkPhysics.cpp
--- a/Engine/Plugins/Runtime/NetworkPredictionExtras/Source/NetworkPredictionExtras/Private/MockNetworkPhysics.cpp
+++ b/Engine/Plugins/Runtime/NetworkPredictionExtras/Source/NetworkPredictionExtras/Private/MockNetworkPhysics.cpp
@@ -148,8 +148,14 @@ namespace UE_NETWORK_PHYSICS
 	bool JumpHack=false;
 	FAutoConsoleVariableRef CVarJumpHack(TEXT("np2.Mock.JumpHack"), JumpHack, TEXT("Make jump not rely on trace which currently causes non determinism"));
 
-	float MockImpulse=false;
-	FAutoConsoleVariableRef CVarMockImpulse(TEXT("np2.Mock.BallImpulse"), JumpHack, TEXT("Make jump not rely on trace which currently causes non determinism"));
+	bool MockImpulse=true;
+	FAutoConsoleVariableRef CVarMockImpulse(TEXT("np2.Mock.BallImpulse"), MockImpulse, TEXT("Make jump not rely on trace which currently causes non determinism"));
+
+	float MockImpulseX=500;
+	FAutoConsoleVariableRef CVarMockImpulseX(TEXT("np2.Mock.BallImpulse.X"), MockImpulseX, TEXT("X magnitude"));
+
+	float MockImpulseZ=300.f;
+	FAutoConsoleVariableRef CVarMockImpulseZ(TEXT("np2.Mock.BallImpulse.Z"), MockImpulseZ, TEXT("Z magnitude"));
 }
 
 void FMockManagedState::AsyncTick(UWorld* World, Chaos::FPhysicsSolver* Solver, const float DeltaSeconds, const int32 SimulationFrame, const int32 LocalStorageFrame)
@@ -159,12 +165,14 @@ void FMockManagedState::AsyncTick(UWorld* World, Chaos::FPhysicsSolver* Solver,
 	{
 		if (auto* PT = Proxy->GetPhysicsThreadAPI())
 		{
-			
 			FVector TracePosition = PT->X();
 			FVector EndPosition = TracePosition + FVector(0.f, 0.f, -100.f);
 			FCollisionShape Shape = FCollisionShape::MakeSphere(250.f);
 			ECollisionChannel CollisionChannel = ECollisionChannel::ECC_WorldStatic; 
 			FCollisionQueryParams QueryParams = FCollisionQueryParams::DefaultQueryParam;
+
+			
+
 			FCollisionResponseParams ResponseParams = FCollisionResponseParams::DefaultResponseParam;
 			FCollisionObjectQueryParams ObjectParams(ECollisionChannel::ECC_PhysicsBody);
 
@@ -172,9 +180,58 @@ void FMockManagedState::AsyncTick(UWorld* World, Chaos::FPhysicsSolver* Solver,
 			const bool bInAir = !UE_NETWORK_PHYSICS::JumpHack && !World->LineTraceSingleByChannel(OutHit, TracePosition, EndPosition, ECollisionChannel::ECC_WorldStatic, QueryParams, ResponseParams);
 			const float UpDot = FVector::DotProduct(PT->R().GetUpVector(), FVector::UpVector);
 
-			//TArray<FOverlapResult> OutOverlaps;
-			//FGenericPhysicsInterface::GeomOverlapMulti(World, Shape, TracePosition, FQuat::Identity, OutOverlaps, ECollisionChannel::ECC_PhysicsBody, QueryParams, ResponseParams, ObjectParams);
+			if (UE_NETWORK_PHYSICS::MockImpulse && PT_State.KickFrame + 100 < SimulationFrame)
+			{
+				FVector Vel = PT->V();
+				if (Vel.SizeSquared() > 1.f && Component)
+				{
+					TArray<FHitResult> SweepHits;
+					
+					// Super gross and sketch
+					UPrimitiveComponent* MyPrimitive = Cast<UPrimitiveComponent>(this->Component->GetOwner()->GetRootComponent());
+					if (MyPrimitive)
+					{
+						QueryParams.AddIgnoredActor(this->Component->GetOwner()); // Super unsafe
+						if (World->SweepMultiByChannel(SweepHits, TracePosition, TracePosition + (Vel * 1.f * DeltaSeconds), PT->R(), ECollisionChannel::ECC_PhysicsBody, MyPrimitive->GetCollisionShape(), QueryParams, ResponseParams))
+						{
+							for (FHitResult& SweepHit : SweepHits)
+							{
+								if (UPrimitiveComponent* HitPrimitive = SweepHit.GetComponent()) // Super unsafe
+								{
+									if (HitPrimitive->IsSimulatingPhysics())
+									{
+										if (UNetworkPhysicsComponent* NetworkPhysicsComp = HitPrimitive->GetOwner()->FindComponentByClass<UNetworkPhysicsComponent>()) // really bad and stupid
+										{
+											if (NetworkPhysicsComp->bCanBeKicked) // We want this on the GT_State but we can't easily map the hit result back to this yet
+											{
+												if (FBodyInstance* Instance = HitPrimitive->GetBodyInstance())
+												{
+													FPhysicsActorHandle HitHandle = Instance->GetPhysicsActorHandle();
+													if (auto* HitPT = HitHandle->GetPhysicsThreadAPI())
+													{
+														FVector Impulse = HitPT->X() - PT->X();
+														Impulse.Z =0.f;
+														Impulse.Normalize();
+														Impulse *= UE_NETWORK_PHYSICS::MockImpulseX;
+														Impulse.Z = UE_NETWORK_PHYSICS::MockImpulseZ;
+
+														//UE_LOG(LogTemp, Warning, TEXT("Applied Force. %s"), *GetNameSafe(HitPrimitive->GetOwner()));
+														HitPT->SetLinearImpulse( Impulse, false );
+														PT_State.KickFrame = SimulationFrame;
+													}
+												}
+
+											}
+										}
+									}
+								}
+							}
+						}
+					}
+				}
+			}
 
+			
 			// ---------------------------------------------------------------------------------------------
 						
 			if (!bInAir)
diff --git a/Engine/Plugins/Runtime/NetworkPredictionExtras/Source/NetworkPredictionExtras/Public/MockNetworkPhysics.h b/Engine/Plugins/Runtime/NetworkPredictionExtras/Source/NetworkPredictionExtras/Public/MockNetworkPhysics.h
--- a/Engine/Plugins/Runtime/NetworkPredictionExtras/Source/NetworkPredictionExtras/Public/MockNetworkPhysics.h
+++ b/Engine/Plugins/Runtime/NetworkPredictionExtras/Source/NetworkPredictionExtras/Public/MockNetworkPhysics.h
@@ -145,6 +145,10 @@ struct FMockState_PT
 	UPROPERTY(BlueprintReadWrite, Category="Mock Object")
 	int32 InAirFrame = 0;
 
+	// Frame we last applied a kick impulse
+	UPROPERTY(BlueprintReadWrite, Category="Mock Object")
+	int32 KickFrame = 0;
+
 	bool NetSerialize(FArchive& Ar, class UPackageMap* Map, bool& bOutSuccess)
 	{
 		Ar << JumpCooldownMS;
@@ -153,12 +157,13 @@ struct FMockState_PT
 		Ar << RecoveryFrame;
 		Ar << JumpStartFrame;
 		Ar << InAirFrame;
+		Ar << KickFrame;
 		return true;
 	}
 
 	bool ShouldReconcile(const FMockState_PT& AuthState) const
 	{
-		return JumpCooldownMS != AuthState.JumpCooldownMS || JumpCount != AuthState.JumpCount || RecoveryFrame != AuthState.RecoveryFrame || JumpStartFrame != AuthState.JumpStartFrame || InAirFrame != AuthState.InAirFrame;
+		return JumpCooldownMS != AuthState.JumpCooldownMS || JumpCount != AuthState.JumpCount || RecoveryFrame != AuthState.RecoveryFrame || JumpStartFrame != AuthState.JumpStartFrame || InAirFrame != AuthState.InAirFrame || KickFrame!= AuthState.KickFrame;
 	}
 };
 
@@ -368,6 +373,9 @@ public:
 	UPROPERTY(EditAnywhere, BlueprintReadWrite, Category="Recorded Inputs")
 	TArray<FMockRecordedInputs>	RecordedInputs;
 
+	UPROPERTY(EditAnywhere, BlueprintReadWrite, Category="Mock State")
+	bool bCanBeKicked=false;
+
 protected:
 
 	// Managed state should not be publically exposed
