Fix static analysis warning

diff --git a/Engine/Plugins/Runtime/NetworkPredictionExtras/Source/NetworkPredictionExtras/Private/MockNetworkPhysics.cpp b/Engine/Plugins/Runtime/NetworkPredictionExtras/Source/NetworkPredictionExtras/Private/MockNetworkPhysics.cpp
--- a/Engine/Plugins/Runtime/NetworkPredictionExtras/Source/NetworkPredictionExtras/Private/MockNetworkPhysics.cpp
+++ b/Engine/Plugins/Runtime/NetworkPredictionExtras/Source/NetworkPredictionExtras/Private/MockNetworkPhysics.cpp
@@ -1159,7 +1159,7 @@ AActor* ANetworkPredictionSpawner::Spawn(FName StreamName)
 		return nullptr;
 	}
 
-	StreamName = (StreamName == NAME_None ? SourceSpawner->RecordedInputs[FMath::Rand() % SourceSpawner->RecordedInputs.Num()].Name : StreamName);
+	StreamName = ((StreamName == NAME_None) ? SourceSpawner->RecordedInputs[FMath::Rand() % SourceSpawner->RecordedInputs.Num()].Name : StreamName);
 
 	if (FMockRecordedInputs* PlaybackRecordedInputs = SourceSpawner->RecordedInputs.FindByKey(StreamName))
 	{
diff --git a/Engine/Source/Editor/PropertyEditor/Private/SPerQualityLevelPropertiesWidget.cpp b/Engine/Source/Editor/PropertyEditor/Private/SPerQualityLevelPropertiesWidget.cpp
--- a/Engine/Source/Editor/PropertyEditor/Private/SPerQualityLevelPropertiesWidget.cpp
+++ b/Engine/Source/Editor/PropertyEditor/Private/SPerQualityLevelPropertiesWidget.cpp
@@ -41,7 +41,7 @@ TSharedRef<SWidget> SOverridePropertiesWidget::MakeOverrideWidget(FName InName,
 
 	TSharedRef<SWidget> Widget =
 		SNew(SBox)
-		.ToolTipText(InName == NAME_None ?
+		.ToolTipText((InName == NAME_None) ?
 			NSLOCTEXT("SPerQualityLevelPropertiesWidget", "DefaultQualityLevelDesc", "This property can have per quality level overrides.\nThis is the default value used when no override has been set for a specific quality level.") :
 			FText::Format(NSLOCTEXT("SPerQualityLevelPropertiesWidget", "QualityLevelDesc", "Override for {0}"), InDisplayText))
 		.Padding(FMargin(0.0f, 2.0f, 4.0f, 2.0f))
