Network Physics
-Temp implementation of ball kick/impulse as a radius check on PT instead of doing scene query, which is against stale GT data, which causes corrections and non determinism.

#rnx

diff --git a/Engine/Plugins/Runtime/NetworkPredictionExtras/Source/NetworkPredictionExtras/Private/MockNetworkPhysics.cpp b/Engine/Plugins/Runtime/NetworkPredictionExtras/Source/NetworkPredictionExtras/Private/MockNetworkPhysics.cpp
--- a/Engine/Plugins/Runtime/NetworkPredictionExtras/Source/NetworkPredictionExtras/Private/MockNetworkPhysics.cpp
+++ b/Engine/Plugins/Runtime/NetworkPredictionExtras/Source/NetworkPredictionExtras/Private/MockNetworkPhysics.cpp
@@ -158,7 +158,7 @@ namespace UE_NETWORK_PHYSICS
 	FAutoConsoleVariableRef CVarMockImpulseZ(TEXT("np2.Mock.BallImpulse.Z"), MockImpulseZ, TEXT("Z magnitude"));
 }
 
-void FMockManagedState::AsyncTick(UWorld* World, Chaos::FPhysicsSolver* Solver, const float DeltaSeconds, const int32 SimulationFrame, const int32 LocalStorageFrame)
+void FMockManagedState::AsyncTick(UWorld* World, Chaos::FPhysicsSolver* Solver, const float DeltaSeconds, const int32 SimulationFrame, const int32 LocalStorageFrame, const TArray<FSingleParticlePhysicsProxy*>& BallProxies)
 {
 	// FIXME: PT_State is the only thing that should be writable here
 	if (ensure(Proxy))
@@ -180,8 +180,32 @@ void FMockManagedState::AsyncTick(UWorld* World, Chaos::FPhysicsSolver* Solver,
 			const bool bInAir = !UE_NETWORK_PHYSICS::JumpHack && !World->LineTraceSingleByChannel(OutHit, TracePosition, EndPosition, ECollisionChannel::ECC_WorldStatic, QueryParams, ResponseParams);
 			const float UpDot = FVector::DotProduct(PT->R().GetUpVector(), FVector::UpVector);
 
-			if (UE_NETWORK_PHYSICS::MockImpulse && PT_State.KickFrame + 100 < SimulationFrame)
+			if (UE_NETWORK_PHYSICS::MockImpulse && PT_State.KickFrame + 10 < SimulationFrame)
 			{
+				for (FSingleParticlePhysicsProxy* BallProxy : BallProxies)
+				{					
+					if (auto* BallPT = BallProxy->GetPhysicsThreadAPI())
+					{
+						const FVector BallLocation = BallPT->X();
+						const float BallRadius = BallPT->Geometry()->BoundingBox().OriginRadius(); //GetRadius();
+						
+						if (BallRadius > 0.f && (FVector::DistSquared(PT->X(), BallLocation) < BallRadius * BallRadius))
+						{
+							FVector Impulse = BallPT->X() - PT->X();
+							Impulse.Z =0.f;
+							Impulse.Normalize();
+							Impulse *= UE_NETWORK_PHYSICS::MockImpulseX;
+							Impulse.Z = UE_NETWORK_PHYSICS::MockImpulseZ;
+
+							//UE_LOG(LogTemp, Warning, TEXT("Applied Force. %s"), *GetNameSafe(HitPrimitive->GetOwner()));
+							BallPT->SetLinearImpulse( Impulse, false );
+							PT_State.KickFrame = SimulationFrame;
+						}
+					}
+				}
+				
+
+				/*
 				FVector Vel = PT->V();
 				if (Vel.SizeSquared() > 1.f && Component)
 				{
@@ -229,6 +253,7 @@ void FMockManagedState::AsyncTick(UWorld* World, Chaos::FPhysicsSolver* Solver,
 						}
 					}
 				}
+				*/
 			}
 
 			
@@ -363,6 +388,8 @@ struct FMockAsyncObjectManagerInput : public Chaos::FSimCallbackInput
 {
 	// One per instance of our physics objects
 	TArray<FMockManagedState> ManagedObjects;
+
+	TArray<FSingleParticlePhysicsProxy*> BallProxies;
 	
 	TWeakObjectPtr<UWorld> World;
 	int32 Timestamp = INDEX_NONE;
@@ -477,7 +504,7 @@ public:
 			PT_Obj.Frame = Frame;
 			SnapshotForGT.Objects.Add(PT_Obj);
 			
-			PT_Obj.AsyncTick(World, PhysicsSolver, GetDeltaTime_Internal(), SimulationFrame, Frame);
+			PT_Obj.AsyncTick(World, PhysicsSolver, GetDeltaTime_Internal(), SimulationFrame, Frame, Input->BallProxies);
 		}
 
 		if (SnapshotForGT.Objects.Num() > 0)
@@ -900,6 +927,7 @@ void FMockObjectManager::ProcessInputs_External(int32 PhysicsStep, int32 LocalFr
 	AsyncInput->Reset();	//only want latest frame's data
 	AsyncInput->World = WeakWorld;
 	AsyncInput->ManagedObjects.Reserve(InMockManagedStates.Num());
+	AsyncInput->BallProxies = BallProxies;
 		
 	for (FMockManagedState* State : InMockManagedStates)
 	{
@@ -925,6 +953,19 @@ void FMockObjectManager::ProcessInputs_External(int32 PhysicsStep, int32 LocalFr
 	}
 }
 
+void FMockObjectManager::RegisterBall(FSingleParticlePhysicsProxy* Proxy)
+{
+	if (ensure(BallProxies.Contains(Proxy) == false))
+	{
+		BallProxies.Add(Proxy);
+	}
+}
+
+void FMockObjectManager::UnregisterBall(FSingleParticlePhysicsProxy* Proxy)
+{
+	BallProxies.Remove(Proxy);
+}
+
 // ========================================================================================
 
 UNetworkPhysicsComponent::UNetworkPhysicsComponent()
@@ -1002,6 +1043,14 @@ void UNetworkPhysicsComponent::InitializeComponent()
 			ReplicatedManagedState.Proxy = PrimitiveComponent->BodyInstance.ActorHandle;
 			MockManager->RegisterManagedMockObject(&ReplicatedManagedState, &InManagedState, &OutManagedState);
 		}
+
+		if (bCanBeKicked)
+		{
+			FMockObjectManager* MockManager = FMockObjectManager::Get(World);
+			checkSlow(MockManager);
+
+			MockManager->RegisterBall(PrimitiveComponent->BodyInstance.ActorHandle);
+		}
 	}
 #endif
 }
@@ -1026,10 +1075,21 @@ void UNetworkPhysicsComponent::EndPlay(const EEndPlayReason::Type EndPlayReason)
 				if (FMockObjectManager* MockManager = FMockObjectManager::Get(World))
 				{
 					//UE_LOG(LogTemp, Warning, TEXT("   Unregistering MockManager. 0x%X"), (int64)&ReplicatedManagedState);
-					MockManager->UnregisterManagedMockObject(&ReplicatedManagedState, &InManagedState, &OutManagedState);				
+					MockManager->UnregisterManagedMockObject(&ReplicatedManagedState, &InManagedState, &OutManagedState);
 				}
 			}
+
+			if (bCanBeKicked)
+			{
+				if (FMockObjectManager* MockManager = FMockObjectManager::Get(World))
+				{
+					MockManager->UnregisterBall(NetworkPhysicsState.Proxy);
+				}
+			}
+
 			Manager->UnregisterPhysicsProxy(&NetworkPhysicsState);
+
+
 		}
 	}
 #endif
diff --git a/Engine/Plugins/Runtime/NetworkPredictionExtras/Source/NetworkPredictionExtras/Public/MockNetworkPhysics.h b/Engine/Plugins/Runtime/NetworkPredictionExtras/Source/NetworkPredictionExtras/Public/MockNetworkPhysics.h
--- a/Engine/Plugins/Runtime/NetworkPredictionExtras/Source/NetworkPredictionExtras/Public/MockNetworkPhysics.h
+++ b/Engine/Plugins/Runtime/NetworkPredictionExtras/Source/NetworkPredictionExtras/Public/MockNetworkPhysics.h
@@ -196,7 +196,7 @@ struct FMockManagedState
 {
 	GENERATED_BODY()
 
-	void AsyncTick(UWorld* World, Chaos::FPhysicsSolver* Solver, const float DeltaSeconds, const int32 SimulationFrame, const int32 LocalStorageFrame);
+	void AsyncTick(UWorld* World, Chaos::FPhysicsSolver* Solver, const float DeltaSeconds, const int32 SimulationFrame, const int32 LocalStorageFrame, const TArray<FSingleParticlePhysicsProxy*>& BallProxies);
 
 	UPROPERTY()
 	int32 Frame = INDEX_NONE;
@@ -271,12 +271,18 @@ public:
 	void PreNetSend(UWorld* World, float DeltaSeconds) override;
 	void ProcessInputs_External(int32 PhysicsStep, int32 LocalFrameOffset, bool& bOutSendClientInputCmd) override;
 
+	void RegisterBall(FSingleParticlePhysicsProxy* Proxy);
+	void UnregisterBall(FSingleParticlePhysicsProxy* Proxy);
+
 private:
 
 	TArray<FMockManagedState*> ReplicatedMockManagedStates;
 	TArray<FMockManagedState*> InMockManagedStates;
 	TArray<FMockManagedState*> OutMockManagedStates;
 	class FMockAsyncObjectManagerCallback* AsyncCallback = nullptr;
+
+	TArray<FSingleParticlePhysicsProxy*> BallProxies;
+
 	
 	TWeakObjectPtr<UWorld> WeakWorld;
 };
