Removed some UE4 references from comments

#rnx

diff --git a/Engine/Plugins/Runtime/GameplayAbilities/Source/GameplayAbilities/Public/GameplayPrediction.h b/Engine/Plugins/Runtime/GameplayAbilities/Source/GameplayAbilities/Public/GameplayPrediction.h
--- a/Engine/Plugins/Runtime/GameplayAbilities/Source/GameplayAbilities/Public/GameplayPrediction.h
+++ b/Engine/Plugins/Runtime/GameplayAbilities/Source/GameplayAbilities/Public/GameplayPrediction.h
@@ -31,7 +31,7 @@ DECLARE_DELEGATE(FPredictionKeyEvent);
  *	-Gameplay Cue events (both from within predictive gameplay effect and on their own)
  *		
  *	-Montages
- *	-Movement (built into UE4 UCharacterMovement)
+ *	-Movement (built into UE UCharacterMovement)
  *	
  *	
  *	Some things we don't predict (most of these we potentially could, but currently dont):
diff --git a/Engine/Plugins/Runtime/NetworkPrediction/Source/NetworkPrediction/Public/NetworkPredictionComponent.h b/Engine/Plugins/Runtime/NetworkPrediction/Source/NetworkPrediction/Public/NetworkPredictionComponent.h
--- a/Engine/Plugins/Runtime/NetworkPrediction/Source/NetworkPrediction/Public/NetworkPredictionComponent.h
+++ b/Engine/Plugins/Runtime/NetworkPrediction/Source/NetworkPrediction/Public/NetworkPredictionComponent.h
@@ -10,7 +10,7 @@
 // ----------------------------------------------------------------------------------------------------------------------------------------------------------------------
 //	UNetworkPredictionComponent
 //	This is the base component for running a TNetworkedSimulationModel through an actor component. This contains the boiler plate hooks into getting the system
-//	initialized and plugged into the UE4 replication system.
+//	initialized and plugged into the UE replication system.
 //
 //	This is an abstract component and cannot function on its own. It must be subclassed and InitializeNetworkPredictionProxy must be implemented.
 //	Ticking and RPC sending will be handled automatically.
diff --git a/Engine/Plugins/Runtime/NetworkPrediction/Source/NetworkPrediction/Public/NetworkPredictionDriver.h b/Engine/Plugins/Runtime/NetworkPrediction/Source/NetworkPrediction/Public/NetworkPredictionDriver.h
--- a/Engine/Plugins/Runtime/NetworkPrediction/Source/NetworkPrediction/Public/NetworkPredictionDriver.h
+++ b/Engine/Plugins/Runtime/NetworkPrediction/Source/NetworkPrediction/Public/NetworkPredictionDriver.h
@@ -19,7 +19,7 @@ namespace Chaos
 	class FRewindData;
 };
 
-// The "driver" represents the UE4 side of the NetworkPrediction System. Typically the driver will be an Unreal Actor or ActorComponent.
+// The "driver" represents the UE side of the NetworkPrediction System. Typically the driver will be an Unreal Actor or ActorComponent.
 // This driver class is defined in the ModelDef (your version of FNetworkPredictionModelDef). For example, AActor or AMyPawn class.
 // While the Simulation class is agnostic to all of this, the driver is responsible for specifying exactly how things should work, 
 // generating input and handling output from the system.
diff --git a/Engine/Plugins/Runtime/NetworkPredictionExtras/Source/NetworkPredictionExtras/Public/FlyingMovementSimulation.h b/Engine/Plugins/Runtime/NetworkPredictionExtras/Source/NetworkPredictionExtras/Public/FlyingMovementSimulation.h
--- a/Engine/Plugins/Runtime/NetworkPredictionExtras/Source/NetworkPredictionExtras/Public/FlyingMovementSimulation.h
+++ b/Engine/Plugins/Runtime/NetworkPredictionExtras/Source/NetworkPredictionExtras/Public/FlyingMovementSimulation.h
@@ -10,7 +10,7 @@
 #include "NetworkPredictionSimulation.h"
 
 // -------------------------------------------------------------------------------------------------------------------------------
-// FlyingMovement: simple flying movement that was based on UE4's FloatingPawnMovement
+// FlyingMovement: simple flying movement that was based on UE's FloatingPawnMovement
 // -------------------------------------------------------------------------------------------------------------------------------
 
 // State the client generates
diff --git a/Engine/Plugins/Runtime/NetworkPredictionExtras/Source/NetworkPredictionExtras/Public/MockNetworkSimulation.h b/Engine/Plugins/Runtime/NetworkPredictionExtras/Source/NetworkPredictionExtras/Public/MockNetworkSimulation.h
--- a/Engine/Plugins/Runtime/NetworkPredictionExtras/Source/NetworkPredictionExtras/Public/MockNetworkSimulation.h
+++ b/Engine/Plugins/Runtime/NetworkPredictionExtras/Source/NetworkPredictionExtras/Public/MockNetworkSimulation.h
@@ -38,7 +38,7 @@ DECLARE_CYCLE_STAT(TEXT("NetworkPrediction_MockSimTick"), STAT_NetworkPrediction
 //
 //	Highlights:
 //		FMockNetworkModel::SimulationTick		The "core update" function of the simulation.
-//		UMockNetworkSimulationComponent			The UE4 Actor Component that anchors the system to an actor.
+//		UMockNetworkSimulationComponent			The UE Actor Component that anchors the system to an actor.
 //
 //	Usage:
 //		You can just add a UMockNetworkSimulationComponent to any ROLE_AutonomousProxy actor yourself (Default Subobject, through blueprints, manually, etc)
diff --git a/Engine/Plugins/Runtime/NetworkPredictionExtras/Source/NetworkPredictionExtras/Public/NetworkPredictionExtrasCharacter.h b/Engine/Plugins/Runtime/NetworkPredictionExtras/Source/NetworkPredictionExtras/Public/NetworkPredictionExtrasCharacter.h
--- a/Engine/Plugins/Runtime/NetworkPredictionExtras/Source/NetworkPredictionExtras/Public/NetworkPredictionExtrasCharacter.h
+++ b/Engine/Plugins/Runtime/NetworkPredictionExtras/Source/NetworkPredictionExtras/Public/NetworkPredictionExtrasCharacter.h
@@ -15,7 +15,7 @@ class UCharacterMotionComponent;
 //
 //	This provides a minimal pawn class that uses UCharacterMotionCompnent. This isn't really intended to be used in shipping games,
 //	rather just to serve as standalone example of using the system, contained completely in the NetworkPredictionExtras plugin.
-//	Apart from the most basic glue/setup, this class provides an example of turning UE4 input event callbacks into the input commands
+//	Apart from the most basic glue/setup, this class provides an example of turning UE input event callbacks into the input commands
 //	that are used by the Character movement simulation. This includes some basic camera/aiming code.
 //
 //	Highlights:
diff --git a/Engine/Plugins/Runtime/NetworkPredictionExtras/Source/NetworkPredictionExtras/Public/NetworkPredictionExtrasFlyingPawn.h b/Engine/Plugins/Runtime/NetworkPredictionExtras/Source/NetworkPredictionExtras/Public/NetworkPredictionExtrasFlyingPawn.h
--- a/Engine/Plugins/Runtime/NetworkPredictionExtras/Source/NetworkPredictionExtras/Public/NetworkPredictionExtrasFlyingPawn.h
+++ b/Engine/Plugins/Runtime/NetworkPredictionExtras/Source/NetworkPredictionExtras/Public/NetworkPredictionExtrasFlyingPawn.h
@@ -15,7 +15,7 @@ class UFlyingMovementComponent;
 //
 //	This provides a minimal pawn class that uses UFlyingMovementCompnent. This isn't really intended to be used in shipping games,
 //	rather just to serve as standalone example of using the system, contained completely in the NetworkPredictionExtras plugin.
-//	Apart from the most basic glue/setup, this class provides an example of turning UE4 input event callbacks into the input commands
+//	Apart from the most basic glue/setup, this class provides an example of turning UE input event callbacks into the input commands
 //	that are used by the flying movement simulation. This includes some basic camera/aiming code.
 //
 //	Highlights:
