Network Prediction plugin: fix for potential use of null UpdatedPrimitive pointer discovered by static analysis

diff --git a/Engine/Plugins/Runtime/NetworkPredictionExtras/Source/NetworkPredictionExtras/Private/CharacterMotionSimulation.cpp b/Engine/Plugins/Runtime/NetworkPredictionExtras/Source/NetworkPredictionExtras/Private/CharacterMotionSimulation.cpp
--- a/Engine/Plugins/Runtime/NetworkPredictionExtras/Source/NetworkPredictionExtras/Private/CharacterMotionSimulation.cpp
+++ b/Engine/Plugins/Runtime/NetworkPredictionExtras/Source/NetworkPredictionExtras/Private/CharacterMotionSimulation.cpp
@@ -327,14 +327,13 @@ void FCharacterMotionSimulation::Movement_Walking(float DeltaSeconds, const TNet
 	if (!Delta.IsNearlyZero(1e-6f))
 	{
 		SafeMoveUpdatedComponent(Delta, OutputQuat, true, Hit, ETeleportType::None);
+	}
 
-		// Hits could invalidate this, though that really sucks to check every time. Should probably avoid allowing this during the simulation so we don't have to keep checking it every action
-		// which is error-prone and easy to overlook.
-		if (!UpdatedPrimitive)
-		{
-			return;
-		}
-
+	// Hits could invalidate this, though that really sucks to check every time. Should probably avoid allowing this during the simulation so we don't have to keep checking it every action
+	// which is error-prone and easy to overlook.
+	if (!UpdatedPrimitive)
+	{
+		return;
 	}
 
 	FindFloor(Input, Output, UpdatedPrimitive->GetComponentLocation(), CachedFloor);
