Network Prediction Extras: fixed possible cases where forward-predicted sim proxies could fight with Actor movement replication if a pawn was not set up properly.

diff --git a/Engine/Plugins/Runtime/NetworkPredictionExtras/Source/NetworkPredictionExtras/Private/NetworkPredictionExtrasCharacter.cpp b/Engine/Plugins/Runtime/NetworkPredictionExtras/Source/NetworkPredictionExtras/Private/NetworkPredictionExtrasCharacter.cpp
--- a/Engine/Plugins/Runtime/NetworkPredictionExtras/Source/NetworkPredictionExtras/Private/NetworkPredictionExtrasCharacter.cpp
+++ b/Engine/Plugins/Runtime/NetworkPredictionExtras/Source/NetworkPredictionExtras/Private/NetworkPredictionExtrasCharacter.cpp
@@ -41,6 +41,8 @@ ANetworkPredictionExtrasCharacter::ANetworkPredictionExtrasCharacter(const FObje
 {
 	CharacterMotionComponent = CreateDefaultSubobject<UCharacterMotionComponent>(Name_CharacterMotionComponent);
 	ensure(CharacterMotionComponent);
+
+	SetReplicateMovement(false);
 }
 
 void ANetworkPredictionExtrasCharacter::BeginPlay()
diff --git a/Engine/Plugins/Runtime/NetworkPredictionExtras/Source/NetworkPredictionExtras/Private/NetworkPredictionExtrasFlyingPawn.cpp b/Engine/Plugins/Runtime/NetworkPredictionExtras/Source/NetworkPredictionExtras/Private/NetworkPredictionExtrasFlyingPawn.cpp
--- a/Engine/Plugins/Runtime/NetworkPredictionExtras/Source/NetworkPredictionExtras/Private/NetworkPredictionExtrasFlyingPawn.cpp
+++ b/Engine/Plugins/Runtime/NetworkPredictionExtras/Source/NetworkPredictionExtras/Private/NetworkPredictionExtrasFlyingPawn.cpp
@@ -41,6 +41,8 @@ ANetworkPredictionExtrasFlyingPawn::ANetworkPredictionExtrasFlyingPawn(const FOb
 {
 	FlyingMovementComponent = CreateDefaultSubobject<UFlyingMovementComponent>(Name_FlyingMovementComponent);
 	ensure(FlyingMovementComponent);
+
+	SetReplicateMovement(false);
 }
 
 void ANetworkPredictionExtrasFlyingPawn::BeginPlay()
