Changed input production for unpossessed example pawns so that default do-nothing inputs are produced instead of endlessly propagating the last input from when it was possessed.

diff --git a/Engine/Plugins/Runtime/NetworkPredictionExtras/Source/NetworkPredictionExtras/Private/NetworkPredictionExtrasCharacter.cpp b/Engine/Plugins/Runtime/NetworkPredictionExtras/Source/NetworkPredictionExtras/Private/NetworkPredictionExtrasCharacter.cpp
--- a/Engine/Plugins/Runtime/NetworkPredictionExtras/Source/NetworkPredictionExtras/Private/NetworkPredictionExtrasCharacter.cpp
+++ b/Engine/Plugins/Runtime/NetworkPredictionExtras/Source/NetworkPredictionExtras/Private/NetworkPredictionExtrasCharacter.cpp
@@ -252,6 +252,12 @@ void ANetworkPredictionExtrasCharacter::ProduceInput(const int32 DeltaMS, FChara
 
 	if (Controller == nullptr)
 	{
+		if (GetLocalRole() == ENetRole::ROLE_Authority && GetRemoteRole() == ENetRole::ROLE_SimulatedProxy)
+		{
+			// If we get here, that means this pawn is not currently possessed and we're choosing to provide default do-nothing input
+			Cmd = FCharacterMotionInputCmd();
+		}
+
 		// We don't have a local controller so we can't run the code below. This is ok. Simulated proxies will just use previous input when extrapolating
 		return;
 	}
diff --git a/Engine/Plugins/Runtime/NetworkPredictionExtras/Source/NetworkPredictionExtras/Private/NetworkPredictionExtrasFlyingPawn.cpp b/Engine/Plugins/Runtime/NetworkPredictionExtras/Source/NetworkPredictionExtras/Private/NetworkPredictionExtrasFlyingPawn.cpp
--- a/Engine/Plugins/Runtime/NetworkPredictionExtras/Source/NetworkPredictionExtras/Private/NetworkPredictionExtrasFlyingPawn.cpp
+++ b/Engine/Plugins/Runtime/NetworkPredictionExtras/Source/NetworkPredictionExtras/Private/NetworkPredictionExtrasFlyingPawn.cpp
@@ -249,9 +249,14 @@ void ANetworkPredictionExtrasFlyingPawn::ProduceInput(const int32 DeltaMS, FFlyi
 		return;
 	}
 
-
 	if (Controller == nullptr)
 	{
+		if (GetLocalRole() == ENetRole::ROLE_Authority && GetRemoteRole() == ENetRole::ROLE_SimulatedProxy)
+		{
+			// If we get here, that means this pawn is not currently possessed and we're choosing to provide default do-nothing input
+			Cmd = FFlyingMovementInputCmd();
+		}
+
 		// We don't have a local controller so we can't run the code below. This is ok. Simulated proxies will just use previous input when extrapolating
 		return;
 	}
