NetworkPredictionExtras: fixed issue where various example pawns will exhibit slow movement during low framerate situations, due to cached input getting consumed and cleared at the faster NP sim rate.

diff --git a/Engine/Plugins/Runtime/NetworkPredictionExtras/Source/NetworkPredictionExtras/Private/NetworkPredictionExtrasCharacter.cpp b/Engine/Plugins/Runtime/NetworkPredictionExtras/Source/NetworkPredictionExtras/Private/NetworkPredictionExtrasCharacter.cpp
--- a/Engine/Plugins/Runtime/NetworkPredictionExtras/Source/NetworkPredictionExtras/Private/NetworkPredictionExtrasCharacter.cpp
+++ b/Engine/Plugins/Runtime/NetworkPredictionExtras/Source/NetworkPredictionExtras/Private/NetworkPredictionExtrasCharacter.cpp
@@ -258,6 +258,9 @@ void ANetworkPredictionExtrasCharacter::ProduceInput(const int32 DeltaMS, FChara
 			Cmd = FCharacterMotionInputCmd();
 		}
 
+		CachedMoveInput = FVector::ZeroVector;
+		CachedLookInput = FVector2D::ZeroVector;
+
 		// We don't have a local controller so we can't run the code below. This is ok. Simulated proxies will just use previous input when extrapolating
 		return;
 	}
@@ -418,11 +421,7 @@ void ANetworkPredictionExtrasCharacter::ProduceInput(const int32 DeltaMS, FChara
 			}
 			break;
 		}
-	}
-				
-
-	CachedMoveInput = FVector::ZeroVector;
-	CachedLookInput = FVector2D::ZeroVector;
+	}				
 }
 
 // ------------------------------------------------------------------------
