Cleaning up various include paths including removing redundant or unnecessary Public and Private paths.

diff --git a/Engine/Plugins/Experimental/Mutable/Source/MutableRuntime/MutableRuntime.Build.cs b/Engine/Plugins/Experimental/Mutable/Source/MutableRuntime/MutableRuntime.Build.cs
--- a/Engine/Plugins/Experimental/Mutable/Source/MutableRuntime/MutableRuntime.Build.cs
+++ b/Engine/Plugins/Experimental/Mutable/Source/MutableRuntime/MutableRuntime.Build.cs
@@ -14,8 +14,6 @@ namespace UnrealBuildTool.Rules
 			IWYUSupport = IWYUSupport.KeepAsIsForNow;
 			//bUseUnity = false;
 
-			PublicIncludePaths.Add(Path.Combine(ModuleDirectory, "Public"));
-
 			PublicDependencyModuleNames.AddRange(
 				new string[] {
 				"Core"
diff --git a/Engine/Plugins/Media/WebMMedia/Source/WebMMediaEditor/WebMMediaEditor.Build.cs b/Engine/Plugins/Media/WebMMedia/Source/WebMMediaEditor/WebMMediaEditor.Build.cs
--- a/Engine/Plugins/Media/WebMMedia/Source/WebMMediaEditor/WebMMediaEditor.Build.cs
+++ b/Engine/Plugins/Media/WebMMedia/Source/WebMMediaEditor/WebMMediaEditor.Build.cs
@@ -14,11 +14,6 @@ namespace UnrealBuildTool.Rules
 					"EditorFramework",
 					"UnrealEd",
 				});
-
-			PrivateIncludePaths.AddRange(
-				new string[] {
-					"WebMMedia/Private",
-				});
 		}
 	}
 }
diff --git a/Engine/Plugins/Runtime/AR/Google/GoogleARCore/Source/GoogleARCoreRendering/GoogleARCoreRendering.Build.cs b/Engine/Plugins/Runtime/AR/Google/GoogleARCore/Source/GoogleARCoreRendering/GoogleARCoreRendering.Build.cs
--- a/Engine/Plugins/Runtime/AR/Google/GoogleARCore/Source/GoogleARCoreRendering/GoogleARCoreRendering.Build.cs
+++ b/Engine/Plugins/Runtime/AR/Google/GoogleARCore/Source/GoogleARCoreRendering/GoogleARCoreRendering.Build.cs
@@ -6,11 +6,6 @@ public class GoogleARCoreRendering : ModuleRules
 {
 	public GoogleARCoreRendering(ReadOnlyTargetRules Target) : base(Target)
 	{
-		PrivateIncludePaths.AddRange(new string[]
-		{
-			"GoogleARCoreBase/Private",
-		});
-			
 		PrivateDependencyModuleNames.AddRange(new string[]
 		{
 			"Core",
diff --git a/Engine/Plugins/Runtime/AR/Google/GoogleARCore/Source/GoogleARCoreRendering/Private/GoogleARCorePassthroughCameraRenderer.cpp b/Engine/Plugins/Runtime/AR/Google/GoogleARCore/Source/GoogleARCoreRendering/Private/GoogleARCorePassthroughCameraRenderer.cpp
--- a/Engine/Plugins/Runtime/AR/Google/GoogleARCore/Source/GoogleARCoreRendering/Private/GoogleARCorePassthroughCameraRenderer.cpp
+++ b/Engine/Plugins/Runtime/AR/Google/GoogleARCore/Source/GoogleARCoreRendering/Private/GoogleARCorePassthroughCameraRenderer.cpp
@@ -17,7 +17,6 @@
 #include "MaterialShader.h"
 #include "MaterialShaderType.h"
 #include "Materials/MaterialRenderProxy.h"
-#include "GoogleARCoreAndroidHelper.h"
 #include "CommonRenderResources.h"
 #include "ARUtilitiesFunctionLibrary.h"
 #include "HAL/IConsoleManager.h"
diff --git a/Engine/Plugins/Runtime/ModelViewViewModel/Source/ModelViewViewModelBlueprint/ModelViewViewModelBlueprint.Build.cs b/Engine/Plugins/Runtime/ModelViewViewModel/Source/ModelViewViewModelBlueprint/ModelViewViewModelBlueprint.Build.cs
--- a/Engine/Plugins/Runtime/ModelViewViewModel/Source/ModelViewViewModelBlueprint/ModelViewViewModelBlueprint.Build.cs
+++ b/Engine/Plugins/Runtime/ModelViewViewModel/Source/ModelViewViewModelBlueprint/ModelViewViewModelBlueprint.Build.cs
@@ -31,13 +31,6 @@ public class ModelViewViewModelBlueprint : ModuleRules
 				"UnrealEd",
 			});
 
-		PrivateIncludePaths.AddRange(
-			new string[]
-			{
-				"ModelViewViewModel/Private",
-			});
-
-
 		PublicDefinitions.Add("UE_MVVM_WITH_VIEWMODEL_EDITOR=0");
 	}
 }
diff --git a/Engine/Plugins/Runtime/NetworkPredictionExtras/Source/NetworkPredictionExtrasLatentLoad/NetworkPredictionExtrasLatentLoad.Build.cs b/Engine/Plugins/Runtime/NetworkPredictionExtras/Source/NetworkPredictionExtrasLatentLoad/NetworkPredictionExtrasLatentLoad.Build.cs
--- a/Engine/Plugins/Runtime/NetworkPredictionExtras/Source/NetworkPredictionExtrasLatentLoad/NetworkPredictionExtrasLatentLoad.Build.cs
+++ b/Engine/Plugins/Runtime/NetworkPredictionExtras/Source/NetworkPredictionExtrasLatentLoad/NetworkPredictionExtrasLatentLoad.Build.cs
@@ -6,12 +6,6 @@ namespace UnrealBuildTool.Rules
 	{
 		public NetworkPredictionExtrasLatentLoad(ReadOnlyTargetRules Target) : base(Target)
 		{
-			PrivateIncludePaths.AddRange(
-				new string[] {
-                    "NetworkPredictionExtras/Private",
-				}
-				);
-
 			PublicDependencyModuleNames.AddRange(
 				new string[]
 				{
diff --git a/Engine/Source/Editor/GameProjectGeneration/GameProjectGeneration.Build.cs b/Engine/Source/Editor/GameProjectGeneration/GameProjectGeneration.Build.cs
--- a/Engine/Source/Editor/GameProjectGeneration/GameProjectGeneration.Build.cs
+++ b/Engine/Source/Editor/GameProjectGeneration/GameProjectGeneration.Build.cs
@@ -6,12 +6,6 @@ public class GameProjectGeneration : ModuleRules
 {
 	public GameProjectGeneration(ReadOnlyTargetRules Target) : base(Target)
 	{
-		PrivateIncludePaths.AddRange(
-			new string[] {
-				"Editor/AddContentDialog/Private",
-			}
-		);
-
 		PublicDependencyModuleNames.AddRange(
 			new string[] {
 				"HardwareTargeting",
diff --git a/Engine/Source/Editor/GameProjectGeneration/Private/Tests/GameProjectAutomationTests.cpp b/Engine/Source/Editor/GameProjectGeneration/Private/Tests/GameProjectAutomationTests.cpp
--- a/Engine/Source/Editor/GameProjectGeneration/Private/Tests/GameProjectAutomationTests.cpp
+++ b/Engine/Source/Editor/GameProjectGeneration/Private/Tests/GameProjectAutomationTests.cpp
@@ -13,7 +13,6 @@
 #include "LevelEditorViewport.h"
 #include "FileHelpers.h"
 #include "ProjectDescriptor.h"
-#include "IContentSource.h"
 #include "GameProjectUtils.h"
 #include "SProjectDialog.h"
 
diff --git a/Engine/Source/Editor/Kismet/Kismet.Build.cs b/Engine/Source/Editor/Kismet/Kismet.Build.cs
--- a/Engine/Source/Editor/Kismet/Kismet.Build.cs
+++ b/Engine/Source/Editor/Kismet/Kismet.Build.cs
@@ -6,13 +6,6 @@ public class Kismet : ModuleRules
 {
 	public Kismet(ReadOnlyTargetRules Target) : base(Target)
 	{
-		PrivateIncludePaths.AddRange(
-			new string[] {
-				"Editor/GraphEditor/Private",
-				"Editor/SceneOutliner/Private",
-			}
-		);
-
 		PrivateIncludePathModuleNames.AddRange(
 			new string[] { 
 				"AssetRegistry", 
diff --git a/Engine/Source/Editor/SequenceRecorder/SequenceRecorder.Build.cs b/Engine/Source/Editor/SequenceRecorder/SequenceRecorder.Build.cs
--- a/Engine/Source/Editor/SequenceRecorder/SequenceRecorder.Build.cs
+++ b/Engine/Source/Editor/SequenceRecorder/SequenceRecorder.Build.cs
@@ -6,12 +6,6 @@ namespace UnrealBuildTool.Rules
 	{
 		public SequenceRecorder(ReadOnlyTargetRules Target) : base(Target)
 		{
-			PrivateIncludePaths.AddRange(
-				new string[] {
-	                "Editor/SceneOutliner/Private",
-				}
-			);
-
 			PublicDependencyModuleNames.AddRange(
 				new string[]
 				{
diff --git a/Engine/Source/Editor/TranslationEditor/TranslationEditor.Build.cs b/Engine/Source/Editor/TranslationEditor/TranslationEditor.Build.cs
--- a/Engine/Source/Editor/TranslationEditor/TranslationEditor.Build.cs
+++ b/Engine/Source/Editor/TranslationEditor/TranslationEditor.Build.cs
@@ -6,12 +6,6 @@ public class TranslationEditor : ModuleRules
 {
 	public TranslationEditor(ReadOnlyTargetRules Target) : base(Target)
 	{
-		PrivateIncludePaths.AddRange(
-			new string[] {
-				"Editor/PropertyEditor/Private",
-			}
-		);
-
 		PublicIncludePathModuleNames.Add("LevelEditor");
 		PublicIncludePathModuleNames.Add("WorkspaceMenuStructure");
         
